Skein / ViewportStatus: Added SourceControlOutlineMenu, which will extend the viewport's SHOW menu with options to toggle source control status outlines on actors on or off.

#rnx

diff --git a/Engine/Source/Developer/SourceControlViewport/Private/SourceControlViewportModule.cpp b/Engine/Source/Developer/SourceControlViewport/Private/SourceControlViewportModule.cpp
new file mode 100644
--- /dev/null
+++ b/Engine/Source/Developer/SourceControlViewport/Private/SourceControlViewportModule.cpp
@@ -0,0 +1,17 @@
+// Copyright Epic Games, Inc. All Rights Reserved.
+
+#include "SourceControlViewportModule.h"
+#include "SourceControlViewportOutlineMenu.h"
+
+void FSourceControlViewportModule::StartupModule()
+{
+	ViewportOutlineMenu = MakeShared<FSourceControlViewportOutlineMenu>();
+	ViewportOutlineMenu->Init();
+}
+
+void FSourceControlViewportModule::ShutdownModule()
+{
+	ViewportOutlineMenu.Reset();
+}
+
+IMPLEMENT_MODULE( FSourceControlViewportModule, SourceControlViewport );
\ No newline at end of file
diff --git a/Engine/Source/Developer/SourceControlViewport/Private/SourceControlViewportModule.h b/Engine/Source/Developer/SourceControlViewport/Private/SourceControlViewportModule.h
new file mode 100644
--- /dev/null
+++ b/Engine/Source/Developer/SourceControlViewport/Private/SourceControlViewportModule.h
@@ -0,0 +1,36 @@
+// Copyright Epic Games, Inc. All Rights Reserved.
+
+#pragma once
+
+#include "CoreMinimal.h"
+#include "Modules/ModuleManager.h"
+
+class FSourceControlViewportOutlineMenu;
+
+class FSourceControlViewportModule : public IModuleInterface
+{
+public:
+	/** IModuleInterface implementation */
+	virtual void StartupModule() override;
+	virtual void ShutdownModule() override;
+
+	/**
+	 * Gets a reference to the source control outline menu module instance.
+	 *
+	 * @return A reference to the source control outline menu module.
+	 */
+	static FSourceControlViewportModule& Get()
+	{
+		static FName SourceControlViewportModule("SourceControlViewport");
+		return FModuleManager::LoadModuleChecked<FSourceControlViewportModule>(SourceControlViewportModule);
+	}
+
+	static FSourceControlViewportModule* TryGet()
+	{
+		static FName SourceControlViewportModule("SourceControlViewport");
+		return FModuleManager::GetModulePtr<FSourceControlViewportModule>(SourceControlViewportModule);
+	}
+
+private:
+	TSharedPtr<FSourceControlViewportOutlineMenu> ViewportOutlineMenu;
+};
diff --git a/Engine/Source/Developer/SourceControlViewport/Private/SourceControlViewportOutlineMenu.cpp b/Engine/Source/Developer/SourceControlViewport/Private/SourceControlViewportOutlineMenu.cpp
new file mode 100644
--- /dev/null
+++ b/Engine/Source/Developer/SourceControlViewport/Private/SourceControlViewportOutlineMenu.cpp
@@ -0,0 +1,229 @@
+// Copyright Epic Games, Inc. All Rights Reserved.
+
+#include "SourceControlViewportOutlineMenu.h"
+#include "RevisionControlStyle/RevisionControlStyle.h"
+#include "HAL/IConsoleManager.h"
+#include "LevelEditorMenuContext.h"
+#include "LevelEditorViewport.h"
+#include "ToolMenus.h"
+
+#define LOCTEXT_NAMESPACE "SourceControlViewportOutlineMenu"
+
+static bool bEnableViewportOutlineMenu = false;
+TAutoConsoleVariable<bool> CVarSourceControlEnableViewportOutlineMenu(
+	TEXT("SourceControl.ViewportOutlineMenu.Enable"),
+	bEnableViewportOutlineMenu,
+	TEXT("Enables an options menu that allows to toggle source control outlines on or off."),
+	ECVF_Default);
+
+static const FName MenuName("LevelEditor.LevelViewportToolbar.Show");
+static const FName SectionName("RevisionControl");
+static const FName SubMenuName("ShowRevisionControlMenu");
+
+FSourceControlViewportOutlineMenu::FSourceControlViewportOutlineMenu()
+{
+}
+
+FSourceControlViewportOutlineMenu::~FSourceControlViewportOutlineMenu()
+{
+	RemoveViewportOutlineMenu();
+}
+
+void FSourceControlViewportOutlineMenu::Init()
+{
+	CVarSourceControlEnableViewportOutlineMenu->AsVariable()->OnChangedDelegate().AddSPLambda(this,
+		[this](IConsoleVariable* EnableViewportOutlineMenu)
+		{
+			if (EnableViewportOutlineMenu->GetBool())
+			{
+				InsertViewportOutlineMenu();
+			}
+			else
+			{
+				RemoveViewportOutlineMenu();
+			}
+		}
+	);
+}
+
+void FSourceControlViewportOutlineMenu::InsertViewportOutlineMenu()
+{
+	if (UToolMenus* ToolMenus = UToolMenus::TryGet())
+	{
+		UToolMenu* Menu = ToolMenus->ExtendMenu(MenuName);
+		if (Menu != nullptr)
+		{
+			Menu->AddDynamicSection(SectionName,
+				FNewToolMenuDelegate::CreateSP(this, &FSourceControlViewportOutlineMenu::PopulateViewportOutlineMenu)
+			);
+		}
+	}
+}
+
+void FSourceControlViewportOutlineMenu::PopulateViewportOutlineMenu(UToolMenu* InMenu)
+{
+	check(InMenu);
+
+	InMenu->AddDynamicSection(SectionName, FNewToolMenuDelegate::CreateLambda(
+		[this](UToolMenu* InMenu)
+		{
+			ULevelViewportToolBarContext* Context = InMenu->FindContext<ULevelViewportToolBarContext>();
+			if (!Context)
+			{
+				return;
+			}
+
+			FLevelEditorViewportClient* ViewportClient = Context->GetLevelViewportClient();
+			if (!ViewportClient || !ViewportClient->IsPerspective())
+			{
+				return;
+			}
+
+			FToolMenuSection& RevisionControlSection = InMenu->FindOrAddSection(SectionName, LOCTEXT("RevisionControl", "Revision Control"));
+			RevisionControlSection.AddDynamicEntry(NAME_None, FNewToolMenuSectionDelegate::CreateLambda(
+				[this, ViewportClient](FToolMenuSection& InSection)
+				{
+					InSection.AddSubMenu(
+						SubMenuName,
+						LOCTEXT("RevisionControlSubMenu", "Status Highlighting"),
+						LOCTEXT("RevisionControlSubMenu_ToolTip", "Toggle revision control status highlights in the viewport on or off."),
+						FNewToolMenuDelegate::CreateLambda(
+							[this, ViewportClient](UToolMenu* InSubMenu)
+							{
+								FToolMenuSection& DefaultSection = InSubMenu->AddSection(NAME_None);
+
+								DefaultSection.AddMenuEntry(
+									NAME_None,
+									LOCTEXT("ShowAll", "Show All"),
+									LOCTEXT("ShowAll_ToolTip", "Enable highlighting for all statuses"),
+									FSlateIcon(),
+									FUIAction(
+										FExecuteAction::CreateSP(this, &FSourceControlViewportOutlineMenu::ShowAll, ViewportClient)
+									),
+									EUserInterfaceActionType::Button
+								);
+
+								DefaultSection.AddMenuEntry(
+									NAME_None,
+									LOCTEXT("HideAll", "Hide All"),
+									LOCTEXT("HideAll_ToolTip", "Disable highlighting for all statuses"),
+									FSlateIcon(),
+									FUIAction(
+										FExecuteAction::CreateSP(this, &FSourceControlViewportOutlineMenu::HideAll, ViewportClient)
+									),
+									EUserInterfaceActionType::Button
+								);
+
+								DefaultSection.AddSeparator(NAME_None);
+
+								DefaultSection.AddMenuEntry(
+									NAME_None,
+									LOCTEXT("HighlightCheckedOutByOtherUser", "Checked Out by Others"),
+									LOCTEXT("HighlightCheckedOutByOtherUser_ToolTip", "Highlight objects that are checked out by someone else."),
+									FSlateIcon(FRevisionControlStyleManager::GetStyleSetName(), "RevisionControl.CheckedOutByOtherUser"),
+									FUIAction(
+										FExecuteAction::CreateSP(this, &FSourceControlViewportOutlineMenu::ToggleHighlight, ViewportClient, ESourceControlStatus::CheckedOutByOtherUser),
+										FCanExecuteAction(),
+										FIsActionChecked::CreateSP(this, &FSourceControlViewportOutlineMenu::IsHighlighted, ViewportClient, ESourceControlStatus::CheckedOutByOtherUser)
+									),
+									EUserInterfaceActionType::ToggleButton
+								);
+
+								DefaultSection.AddMenuEntry(
+									NAME_None,
+									LOCTEXT("HighlightNotAtHeadRevision", "Out of Date"),
+									LOCTEXT("HighlightNotAtHeadRevision_ToolTip", "Highlight objects that are not at the latest revision."),
+									FSlateIcon(FRevisionControlStyleManager::GetStyleSetName(), "RevisionControl.NotAtHeadRevision"),
+									FUIAction(
+										FExecuteAction::CreateSP(this, &FSourceControlViewportOutlineMenu::ToggleHighlight, ViewportClient, ESourceControlStatus::NotAtHeadRevision),
+										FCanExecuteAction(),
+										FIsActionChecked::CreateSP(this, &FSourceControlViewportOutlineMenu::IsHighlighted, ViewportClient, ESourceControlStatus::NotAtHeadRevision)
+									),
+									EUserInterfaceActionType::ToggleButton
+								);
+
+								DefaultSection.AddMenuEntry(
+									NAME_None,
+									LOCTEXT("HighlightCheckedOut", "Checked Out by Me"),
+									LOCTEXT("HighlightCheckedOut_ToolTip", "Highlight objects that are checked out by me."),
+									FSlateIcon(FRevisionControlStyleManager::GetStyleSetName(), "RevisionControl.CheckedOut"),
+									FUIAction(
+										FExecuteAction::CreateSP(this, &FSourceControlViewportOutlineMenu::ToggleHighlight, ViewportClient, ESourceControlStatus::CheckedOut),
+										FCanExecuteAction(),
+										FIsActionChecked::CreateSP(this, &FSourceControlViewportOutlineMenu::IsHighlighted, ViewportClient, ESourceControlStatus::CheckedOut)
+									),
+									EUserInterfaceActionType::ToggleButton
+								);
+
+								DefaultSection.AddMenuEntry(
+									NAME_None,
+									LOCTEXT("HighlightOpenForAdd", "Newly Added"),
+									LOCTEXT("HighlightOpenForAdd_ToolTip", "Highlight objects that have been added by me."),
+									FSlateIcon(FRevisionControlStyleManager::GetStyleSetName(), "RevisionControl.OpenForAdd"),
+									FUIAction(
+										FExecuteAction::CreateSP(this, &FSourceControlViewportOutlineMenu::ToggleHighlight, ViewportClient, ESourceControlStatus::OpenForAdd),
+										FCanExecuteAction(),
+										FIsActionChecked::CreateSP(this, &FSourceControlViewportOutlineMenu::IsHighlighted, ViewportClient, ESourceControlStatus::OpenForAdd)
+									),
+									EUserInterfaceActionType::ToggleButton
+								);
+							}
+						),
+						false,
+						FSlateIcon(FAppStyle::Get().GetStyleSetName(), "ShowFlagsMenu.SubMenu.RevisionControl")
+					);
+				}));
+		}
+	)
+	);
+}
+
+void FSourceControlViewportOutlineMenu::RemoveViewportOutlineMenu()
+{
+	if (UToolMenus* ToolMenus = UToolMenus::TryGet())
+	{
+		UToolMenu* Menu = ToolMenus->ExtendMenu(MenuName);
+		if (Menu != nullptr)
+		{
+			Menu->RemoveSection(SectionName);
+		}
+	}
+}
+
+void FSourceControlViewportOutlineMenu::ShowAll(FLevelEditorViewportClient* ViewportClient)
+{
+	ensure(ViewportClient);
+
+	SourceControlViewportUtils::SetFeedbackEnabled(ViewportClient, ESourceControlStatus::CheckedOutByOtherUser, /*bEnabled=*/true);
+	SourceControlViewportUtils::SetFeedbackEnabled(ViewportClient, ESourceControlStatus::NotAtHeadRevision, /*bEnabled=*/true);
+	SourceControlViewportUtils::SetFeedbackEnabled(ViewportClient, ESourceControlStatus::CheckedOut, /*bEnabled=*/true);
+	SourceControlViewportUtils::SetFeedbackEnabled(ViewportClient, ESourceControlStatus::OpenForAdd, /*bEnabled=*/true);
+}
+
+void FSourceControlViewportOutlineMenu::HideAll(FLevelEditorViewportClient* ViewportClient)
+{
+	ensure(ViewportClient);
+
+	SourceControlViewportUtils::SetFeedbackEnabled(ViewportClient, ESourceControlStatus::CheckedOutByOtherUser, /*bEnabled=*/false);
+	SourceControlViewportUtils::SetFeedbackEnabled(ViewportClient, ESourceControlStatus::NotAtHeadRevision, /*bEnabled=*/false);
+	SourceControlViewportUtils::SetFeedbackEnabled(ViewportClient, ESourceControlStatus::CheckedOut, /*bEnabled=*/false);
+	SourceControlViewportUtils::SetFeedbackEnabled(ViewportClient, ESourceControlStatus::OpenForAdd, /*bEnabled=*/false);
+}
+
+void FSourceControlViewportOutlineMenu::ToggleHighlight(FLevelEditorViewportClient* ViewportClient, ESourceControlStatus Status)
+{
+	ensure(ViewportClient);
+
+	bool bOld = SourceControlViewportUtils::GetFeedbackEnabled(ViewportClient, Status);
+	bool bNew = bOld ? false : true;
+	SourceControlViewportUtils::SetFeedbackEnabled(ViewportClient, Status, bNew);
+}
+
+bool FSourceControlViewportOutlineMenu::IsHighlighted(FLevelEditorViewportClient* ViewportClient, ESourceControlStatus Status) const
+{
+	ensure(ViewportClient);
+
+	return SourceControlViewportUtils::GetFeedbackEnabled(ViewportClient, Status);
+}
+
+#undef LOCTEXT_NAMESPACE
\ No newline at end of file
diff --git a/Engine/Source/Developer/SourceControlViewport/Private/SourceControlViewportOutlineMenu.h b/Engine/Source/Developer/SourceControlViewport/Private/SourceControlViewportOutlineMenu.h
new file mode 100644
--- /dev/null
+++ b/Engine/Source/Developer/SourceControlViewport/Private/SourceControlViewportOutlineMenu.h
@@ -0,0 +1,32 @@
+// Copyright Epic Games, Inc. All Rights Reserved.
+
+#pragma once
+
+#include "CoreMinimal.h"
+#include "SourceControlViewportUtils.h"
+
+class FLevelEditorViewportClient;
+class UToolMenu;
+
+// Adds an options menu to the Viewport's SHOW pill.
+class FSourceControlViewportOutlineMenu : public TSharedFromThis<FSourceControlViewportOutlineMenu, ESPMode::ThreadSafe>
+{
+public:
+	FSourceControlViewportOutlineMenu();
+	~FSourceControlViewportOutlineMenu();
+
+public:
+	void Init();
+
+private:
+	void InsertViewportOutlineMenu();
+	void PopulateViewportOutlineMenu(UToolMenu* InMenu);
+	void RemoveViewportOutlineMenu();
+
+private:
+	void ShowAll(FLevelEditorViewportClient* ViewportClient);
+	void HideAll(FLevelEditorViewportClient* ViewportClient);
+
+	void ToggleHighlight(FLevelEditorViewportClient* ViewportClient, ESourceControlStatus Status);
+	bool IsHighlighted(FLevelEditorViewportClient* ViewportClient, ESourceControlStatus Status) const;
+};
\ No newline at end of file
diff --git a/Engine/Source/Developer/SourceControlViewport/Private/SourceControlViewportUtils.cpp b/Engine/Source/Developer/SourceControlViewport/Private/SourceControlViewportUtils.cpp
new file mode 100644
--- /dev/null
+++ b/Engine/Source/Developer/SourceControlViewport/Private/SourceControlViewportUtils.cpp
@@ -0,0 +1,119 @@
+// Copyright Epic Games, Inc. All Rights Reserved.
+
+#include "SourceControlViewportUtils.h"
+#include "Settings/EditorStyleSettings.h"
+#include "LevelEditorViewport.h"
+
+namespace SourceControlViewportUtils
+{
+
+// Helper method to get the outline color index for a setting.
+static uint32 GetOutlineColorIndex(ESourceControlStatus Status)
+{
+	FLinearColor ColorToUse;
+	switch (Status)
+	{
+	case ESourceControlStatus::CheckedOutByOtherUser:
+		ColorToUse = FColor::Red;
+		break;
+	case ESourceControlStatus::NotAtHeadRevision:
+		ColorToUse = FColor::Yellow;
+		break;
+	case ESourceControlStatus::CheckedOut:
+		ColorToUse = FColor::Blue;
+		break;
+	case ESourceControlStatus::OpenForAdd:
+		ColorToUse = FColor::Green;
+		break;
+	default:
+		checkNoEntry();
+		break;
+	}
+
+	// The available colors are defined in the UEditorStyleSettings::AdditionalSelectionColors array.
+	// The loop below gets the selection color index closest to the desired color.
+
+	int32 MinIndex = -1;
+	float MinDist = UE_MAX_FLT;
+
+	const UEditorStyleSettings* Settings = GetDefault<UEditorStyleSettings>();
+	for (int32 Index = 0; Index < UE_ARRAY_COUNT(Settings->AdditionalSelectionColors); ++Index)
+	{
+		FLinearColor Color = Settings->AdditionalSelectionColors[Index];
+
+		float Dist = FLinearColor::Dist(ColorToUse, Color);
+		if (Dist < MinDist)
+		{
+			MinIndex = Index;
+			MinDist = Dist;
+		}
+	}
+
+	return MinIndex;
+}
+
+bool GetOutlineSetting(FViewportClient* ViewportClient, ESourceControlStatus Status)
+{
+	uint32 ColorIndex = GetOutlineColorIndex(Status);
+	switch (ColorIndex)
+	{
+	case 0:
+		return ViewportClient->GetEngineShowFlags()->SelectionOutlineColor0;
+	case 1:
+		return ViewportClient->GetEngineShowFlags()->SelectionOutlineColor1;
+	case 2:
+		return ViewportClient->GetEngineShowFlags()->SelectionOutlineColor2;
+	case 3:
+		return ViewportClient->GetEngineShowFlags()->SelectionOutlineColor3;
+	case 4:
+		return ViewportClient->GetEngineShowFlags()->SelectionOutlineColor4;
+	case 5:
+		return ViewportClient->GetEngineShowFlags()->SelectionOutlineColor5;
+	default:
+		ensure(false);
+		break;
+	}
+
+	return false;
+}
+
+void SetOutlineSetting(FViewportClient* ViewportClient, ESourceControlStatus Status, bool bEnabled)
+{
+	uint32 ColorIndex = GetOutlineColorIndex(Status);
+	switch (ColorIndex)
+	{
+	case 0:
+		ViewportClient->GetEngineShowFlags()->SelectionOutlineColor0 = bEnabled;
+		break;
+	case 1:
+		ViewportClient->GetEngineShowFlags()->SelectionOutlineColor1 = bEnabled;
+		break;
+	case 2:
+		ViewportClient->GetEngineShowFlags()->SelectionOutlineColor2 = bEnabled;
+		break;
+	case 3:
+		ViewportClient->GetEngineShowFlags()->SelectionOutlineColor3 = bEnabled;
+		break;
+	case 4:
+		ViewportClient->GetEngineShowFlags()->SelectionOutlineColor4 = bEnabled;
+		break;
+	case 5:
+		ViewportClient->GetEngineShowFlags()->SelectionOutlineColor5 = bEnabled;
+		break;
+	default:
+		ensure(false);
+		break;
+	}
+}
+
+bool GetFeedbackEnabled(FViewportClient* ViewportClient, ESourceControlStatus Status)
+{
+	return GetOutlineSetting(ViewportClient, Status);
+}
+
+void SetFeedbackEnabled(FViewportClient* ViewportClient, ESourceControlStatus Status, bool bEnabled)
+{
+	SetOutlineSetting(ViewportClient, Status, bEnabled);
+}
+
+}
\ No newline at end of file
diff --git a/Engine/Source/Developer/SourceControlViewport/Private/SourceControlViewportUtils.h b/Engine/Source/Developer/SourceControlViewport/Private/SourceControlViewportUtils.h
new file mode 100644
--- /dev/null
+++ b/Engine/Source/Developer/SourceControlViewport/Private/SourceControlViewportUtils.h
@@ -0,0 +1,35 @@
+// Copyright Epic Games, Inc. All Rights Reserved.
+
+#pragma once
+
+#include "CoreMinimal.h"
+#include "UObject/Object.h"
+
+class FViewportClient;
+
+UENUM()
+enum class ESourceControlStatus : uint8
+{
+	/* The file is checked out by another user */
+	CheckedOutByOtherUser,
+
+	/* The file is not at the latest revision */
+	NotAtHeadRevision,
+
+	/* The file is checked out by self */
+	CheckedOut,
+
+	/* The file is added by self */
+	OpenForAdd,
+};
+
+namespace SourceControlViewportUtils
+{
+
+// Helper method to get the viewport feedback value from the viewport.
+bool GetFeedbackEnabled(FViewportClient* ViewportClient, ESourceControlStatus Status);
+
+// Helper method to set the viewport feedback on or off on the viewport.
+void SetFeedbackEnabled(FViewportClient* ViewportClient, ESourceControlStatus Status, bool bEnabled);
+
+}
\ No newline at end of file
diff --git a/Engine/Source/Developer/SourceControlViewport/SourceControlViewport.Build.cs b/Engine/Source/Developer/SourceControlViewport/SourceControlViewport.Build.cs
new file mode 100644
--- /dev/null
+++ b/Engine/Source/Developer/SourceControlViewport/SourceControlViewport.Build.cs
@@ -0,0 +1,23 @@
+// Copyright Epic Games, Inc. All Rights Reserved.
+
+using UnrealBuildTool;
+
+public class SourceControlViewport : ModuleRules
+{
+	public SourceControlViewport(ReadOnlyTargetRules Target) : base(Target)
+	{
+		PrivateDependencyModuleNames.AddRange(
+			new string[] {
+				"Core",
+				"CoreUObject",
+				"Engine",
+				"LevelEditor",
+				"Slate",
+				"SlateCore",
+				"SourceControl",
+				"ToolMenus",
+				"UnrealEd"
+			}
+		);
+	}
+}
diff --git a/Engine/Source/Editor/EditorStyle/Private/StarshipStyle.cpp b/Engine/Source/Editor/EditorStyle/Private/StarshipStyle.cpp
--- a/Engine/Source/Editor/EditorStyle/Private/StarshipStyle.cpp
+++ b/Engine/Source/Editor/EditorStyle/Private/StarshipStyle.cpp
@@ -4675,7 +4675,10 @@ void FStarshipEditorStyle::FStyle::SetupLevelEditorStyle()
 		Set("ShowFlagsMenu.SubMenu.FoliageTypes", new IMAGE_BRUSH_SVG("Starship/Common/FoliageTypes", Icon16x16));
 		Set("ShowFlagsMenu.SubMenu.Sprites", new IMAGE_BRUSH_SVG("Starship/Common/Sprite", Icon16x16));
 		Set("ShowFlagsMenu.SubMenu.HLODs", new IMAGE_BRUSH_SVG("Starship/Common/HierarchicalLOD", Icon16x16));
+		Set("ShowFlagsMenu.SubMenu.RevisionControl", new CORE_IMAGE_BRUSH_SVG("Starship/SourceControl/Status/RevisionControl", Icon16x16));
 	}
+	
+
 #endif // WITH_EDITOR || (IS_PROGRAM && WITH_UNREAL_DEVELOPER_TOOLS)
 
 	// Mobility Icons
