Skein / ViewportStatus: Added Outliner settings analytics.

#rnx

diff --git a/Engine/Source/Developer/SourceControlViewport/Private/SourceControlViewportOutlineMenu.cpp b/Engine/Source/Developer/SourceControlViewport/Private/SourceControlViewportOutlineMenu.cpp
--- a/Engine/Source/Developer/SourceControlViewport/Private/SourceControlViewportOutlineMenu.cpp
+++ b/Engine/Source/Developer/SourceControlViewport/Private/SourceControlViewportOutlineMenu.cpp
@@ -3,6 +3,7 @@
 #include "SourceControlViewportOutlineMenu.h"
 #include "RevisionControlStyle/RevisionControlStyle.h"
 #include "HAL/IConsoleManager.h"
+#include "EngineAnalytics.h"
 #include "LevelEditorMenuContext.h"
 #include "LevelEditorViewport.h"
 #include "ToolMenus.h"
@@ -198,6 +199,8 @@ void FSourceControlViewportOutlineMenu::ShowAll(FLevelEditorViewportClient* View
 	SourceControlViewportUtils::SetFeedbackEnabled(ViewportClient, ESourceControlStatus::NotAtHeadRevision, /*bEnabled=*/true);
 	SourceControlViewportUtils::SetFeedbackEnabled(ViewportClient, ESourceControlStatus::CheckedOut, /*bEnabled=*/true);
 	SourceControlViewportUtils::SetFeedbackEnabled(ViewportClient, ESourceControlStatus::OpenForAdd, /*bEnabled=*/true);
+	
+	RecordToggleEvent(TEXT("All"), /*bEnabled=*/true);
 }
 
 void FSourceControlViewportOutlineMenu::HideAll(FLevelEditorViewportClient* ViewportClient)
@@ -208,6 +211,8 @@ void FSourceControlViewportOutlineMenu::HideAll(FLevelEditorViewportClient* View
 	SourceControlViewportUtils::SetFeedbackEnabled(ViewportClient, ESourceControlStatus::NotAtHeadRevision, /*bEnabled=*/false);
 	SourceControlViewportUtils::SetFeedbackEnabled(ViewportClient, ESourceControlStatus::CheckedOut, /*bEnabled=*/false);
 	SourceControlViewportUtils::SetFeedbackEnabled(ViewportClient, ESourceControlStatus::OpenForAdd, /*bEnabled=*/false);
+
+	RecordToggleEvent(TEXT("All"), /*bEnabled=*/false);
 }
 
 void FSourceControlViewportOutlineMenu::ToggleHighlight(FLevelEditorViewportClient* ViewportClient, ESourceControlStatus Status)
@@ -217,6 +222,10 @@ void FSourceControlViewportOutlineMenu::ToggleHighlight(FLevelEditorViewportClie
 	bool bOld = SourceControlViewportUtils::GetFeedbackEnabled(ViewportClient, Status);
 	bool bNew = bOld ? false : true;
 	SourceControlViewportUtils::SetFeedbackEnabled(ViewportClient, Status, bNew);
+
+	FString EnumValueWithoutType = UEnum::GetValueAsString(Status)
+		.Replace(TEXT("ESourceControlStatus::"), TEXT(""));
+	RecordToggleEvent(EnumValueWithoutType, bNew);
 }
 
 bool FSourceControlViewportOutlineMenu::IsHighlighted(FLevelEditorViewportClient* ViewportClient, ESourceControlStatus Status) const
@@ -226,4 +235,14 @@ bool FSourceControlViewportOutlineMenu::IsHighlighted(FLevelEditorViewportClient
 	return SourceControlViewportUtils::GetFeedbackEnabled(ViewportClient, Status);
 }
 
+void FSourceControlViewportOutlineMenu::RecordToggleEvent(const FString& Param, bool bEnabled) const
+{
+	if (FEngineAnalytics::IsAvailable())
+	{
+		FEngineAnalytics::GetProvider().RecordEvent(
+			TEXT("Editor.Usage.SourceControl.OutlineSettings"), Param, bEnabled ? TEXT("True") : TEXT("False")
+		);
+	}
+}
+
 #undef LOCTEXT_NAMESPACE
\ No newline at end of file
diff --git a/Engine/Source/Developer/SourceControlViewport/Private/SourceControlViewportOutlineMenu.h b/Engine/Source/Developer/SourceControlViewport/Private/SourceControlViewportOutlineMenu.h
--- a/Engine/Source/Developer/SourceControlViewport/Private/SourceControlViewportOutlineMenu.h
+++ b/Engine/Source/Developer/SourceControlViewport/Private/SourceControlViewportOutlineMenu.h
@@ -29,4 +29,7 @@ private:
 
 	void ToggleHighlight(FLevelEditorViewportClient* ViewportClient, ESourceControlStatus Status);
 	bool IsHighlighted(FLevelEditorViewportClient* ViewportClient, ESourceControlStatus Status) const;
+
+private:
+	void RecordToggleEvent(const FString& Param, bool bEnabled) const;
 };
\ No newline at end of file
