Skein / ViewportStatus: Added Viewport Tooltips.

#rnx

diff --git a/Engine/Source/Developer/SourceControlViewport/Private/SourceControlViewportModule.cpp b/Engine/Source/Developer/SourceControlViewport/Private/SourceControlViewportModule.cpp
--- a/Engine/Source/Developer/SourceControlViewport/Private/SourceControlViewportModule.cpp
+++ b/Engine/Source/Developer/SourceControlViewport/Private/SourceControlViewportModule.cpp
@@ -2,15 +2,19 @@
 
 #include "SourceControlViewportModule.h"
 #include "SourceControlViewportOutlineMenu.h"
+#include "SourceControlViewportToolTips.h"
 
 void FSourceControlViewportModule::StartupModule()
 {
 	ViewportOutlineMenu = MakeShared<FSourceControlViewportOutlineMenu>();
 	ViewportOutlineMenu->Init();
+	ViewportToolTips = MakeShared<FSourceControlViewportToolTips>();
+	ViewportToolTips->Init();
 }
 
 void FSourceControlViewportModule::ShutdownModule()
 {
+	ViewportToolTips.Reset();
 	ViewportOutlineMenu.Reset();
 }
 
diff --git a/Engine/Source/Developer/SourceControlViewport/Private/SourceControlViewportModule.h b/Engine/Source/Developer/SourceControlViewport/Private/SourceControlViewportModule.h
--- a/Engine/Source/Developer/SourceControlViewport/Private/SourceControlViewportModule.h
+++ b/Engine/Source/Developer/SourceControlViewport/Private/SourceControlViewportModule.h
@@ -6,6 +6,7 @@
 #include "Modules/ModuleManager.h"
 
 class FSourceControlViewportOutlineMenu;
+class FSourceControlViewportToolTips;
 
 class FSourceControlViewportModule : public IModuleInterface
 {
@@ -33,4 +34,5 @@ public:
 
 private:
 	TSharedPtr<FSourceControlViewportOutlineMenu> ViewportOutlineMenu;
+	TSharedPtr<FSourceControlViewportToolTips> ViewportToolTips;
 };
diff --git a/Engine/Source/Developer/SourceControlViewport/Private/SourceControlViewportToolTips.cpp b/Engine/Source/Developer/SourceControlViewport/Private/SourceControlViewportToolTips.cpp
new file mode 100644
--- /dev/null
+++ b/Engine/Source/Developer/SourceControlViewport/Private/SourceControlViewportToolTips.cpp
@@ -0,0 +1,274 @@
+// Copyright Epic Games, Inc. All Rights Reserved.
+
+#include "SourceControlViewportToolTips.h"
+#include "SourceControlViewportUtils.h"
+#include "Framework/Application/SlateApplication.h"
+#include "Widgets/SToolTip.h"
+#include "Widgets/SCanvas.h"
+#include "Fonts/FontMeasure.h"
+#include "HAL/IConsoleManager.h"
+#include "EngineUtils.h"
+#include "LevelEditor.h"
+#include "SourceControlHelpers.h"
+#include "ISourceControlModule.h"
+#include "ISourceControlProvider.h"
+#include "SLevelViewport.h"
+#include "GameFramework/Actor.h"
+
+#define LOCTEXT_NAMESPACE "SourceControlViewportToolTips"
+
+static bool bEnableViewportToolTips = false;
+TAutoConsoleVariable<bool> CVarSourceControlEnableViewportToolTips(
+	TEXT("SourceControl.ViewportToolTips.Enable"),
+	bEnableViewportToolTips,
+	TEXT("Enables source control tooltips in the viewport."),
+	ECVF_Default);
+
+FSourceControlViewportToolTips::FSourceControlViewportToolTips()
+{
+}
+
+FSourceControlViewportToolTips::~FSourceControlViewportToolTips()
+{
+	FTSTicker::GetCoreTicker().RemoveTicker(TickHandle);
+	TickHandle.Reset();
+}
+
+void FSourceControlViewportToolTips::Init()
+{
+	TickHandle = FTSTicker::GetCoreTicker().AddTicker(
+		FTickerDelegate::CreateRaw(this, &FSourceControlViewportToolTips::Tick)
+	);
+}
+
+bool FSourceControlViewportToolTips::Tick(float DeltaTime)
+{
+	UpdateCanvas(DeltaTime);
+	return true;
+}
+
+void FSourceControlViewportToolTips::UpdateCanvas(float DeltaTime)
+{
+	if (CanvasWidget.IsValid())
+	{
+		// Remove the canvas if the active viewport changed.
+		if (GCurrentLevelEditingViewportClient != nullptr)
+		{
+			TSharedPtr<SLevelViewport> LevelViewportWidget = StaticCastSharedPtr<SLevelViewport>(GCurrentLevelEditingViewportClient->GetEditorViewportWidget());
+			if (LevelViewportWidget != ViewportWidget)
+			{
+				RemoveCanvas();
+			}
+		}
+	}
+
+	bool bEnabled = CVarSourceControlEnableViewportToolTips.GetValueOnGameThread();
+
+	if (CanvasWidget.IsValid() && !bEnabled)
+	{
+		// Remove the canvas if the CVar got disabled.
+		RemoveCanvas();
+	}
+
+	if (!CanvasWidget.IsValid() && bEnabled)
+	{
+		// Insert the canvas if the active viewport is perspective.
+		InsertCanvas();
+	}
+
+	if (CanvasWidget.IsValid() && ToolTipWidget.IsValid() && ViewportWidget.IsValid())
+	{
+		FViewport* Viewport = ViewportWidget->GetViewportClient()->Viewport;
+		check(Viewport);
+
+		int32 MouseX = Viewport->GetMouseX();
+		int32 MouseY = Viewport->GetMouseY();
+		if (ActorMouseX != MouseX || ActorMouseY != MouseY)
+		{
+			Actor.Reset();
+			UpdateToolTip();
+
+			ActorMouseX = MouseX;
+			ActorMouseY = MouseY;
+		}
+
+		HActor* ActorHitProxy = HitProxyCast<HActor>(Viewport->GetHitProxy(MouseX, MouseY));
+		if (ActorHitProxy == nullptr)
+		{
+			Actor.Reset();
+			UpdateToolTip();
+		}
+		if (ActorHitProxy != nullptr && ActorHitProxy->Actor != Actor)
+		{
+			Actor.Reset();
+			UpdateToolTip();
+
+			Actor = MakeWeakObjectPtr(ActorHitProxy->Actor);
+			ActorTime = FPlatformTime::Seconds();
+			DelayTime = 0;
+		}
+
+		if (Actor.IsValid())
+		{
+			DelayTime += DeltaTime;
+			if (DelayTime >= 0.5f)
+			{
+				UpdateToolTip();
+			}
+		}
+	}
+}
+
+void FSourceControlViewportToolTips::UpdateToolTip()
+{
+	ToolTipText = FText::GetEmpty();
+
+	if (Actor.IsValid())
+	{
+		if (UPackage* Package = Actor->GetPackage())
+		{
+			FString SourceFileName = SourceControlHelpers::PackageFilename(Package);
+
+			ISourceControlModule& SourceControlModule = ISourceControlModule::Get();
+			ISourceControlProvider& SourceControlProvider = SourceControlModule.GetProvider();
+			if (SourceControlModule.IsEnabled())
+			{
+				FSourceControlStatePtr State = SourceControlProvider.GetState(SourceFileName, EStateCacheUsage::Use);
+				if (State.IsValid() && State->IsSourceControlled())
+				{
+					FString Who;
+
+					bool bNotAtHeadRevision = !State->IsCurrent();
+					bool bCheckedOutByOtherUser = State->IsCheckedOutOther(&Who);
+					bool bCheckedOut = State->IsCheckedOut();
+					bool bOpenForAdd = State->IsAdded();
+
+					bool bNotAtHeadRevisionEnabled = SourceControlViewportUtils::GetFeedbackEnabled(ViewportWidget->GetViewportClient().Get(), ESourceControlStatus::NotAtHeadRevision);
+					bool bCheckedOutByOtherUserEnabled = SourceControlViewportUtils::GetFeedbackEnabled(ViewportWidget->GetViewportClient().Get(), ESourceControlStatus::CheckedOutByOtherUser);
+					bool bCheckedOutEnabled = SourceControlViewportUtils::GetFeedbackEnabled(ViewportWidget->GetViewportClient().Get(), ESourceControlStatus::CheckedOut);
+					bool bOpenForAddEnabled = SourceControlViewportUtils::GetFeedbackEnabled(ViewportWidget->GetViewportClient().Get(), ESourceControlStatus::OpenForAdd);
+
+					if (bNotAtHeadRevision && bNotAtHeadRevisionEnabled)
+					{
+						ToolTipText = LOCTEXT("NotAtHeadRevision", "File(s) out of sync");
+					}
+					else if (bCheckedOutByOtherUser && bCheckedOutByOtherUserEnabled)
+					{
+						ToolTipText = FText::Format(LOCTEXT("CheckedOutOtherUser", "File(s) checked out by {0}"), FText::FromString(Who));
+					}
+					else if (bCheckedOut && bCheckedOutEnabled)
+					{
+						ToolTipText = LOCTEXT("CheckedOut", "File(s) checked out by you");
+					}
+					else if (bOpenForAdd && bOpenForAddEnabled)
+					{
+						ToolTipText = LOCTEXT("OpenForAdd", "File(s) added by you");
+					}
+				}
+			}
+		}
+	}
+
+	ToolTipWidget->SetVisibility(ToolTipText.IsEmpty() ? EVisibility::Hidden : EVisibility::Visible);
+}
+
+void FSourceControlViewportToolTips::InsertCanvas()
+{
+	check(!CanvasWidget.IsValid());
+	check(!ToolTipWidget.IsValid());
+	check(!ViewportWidget.IsValid());
+
+	if (GCurrentLevelEditingViewportClient != nullptr && GCurrentLevelEditingViewportClient->IsPerspective())
+	{
+		TSharedPtr<SLevelViewport> LevelViewport = StaticCastSharedPtr<SLevelViewport>(GCurrentLevelEditingViewportClient->GetEditorViewportWidget());
+		if (LevelViewport.IsValid())
+		{
+			CanvasWidget = SNew(SCanvas)
+			+SCanvas::Slot()
+				.Position(this, &FSourceControlViewportToolTips::GetToolTipPosition)
+				.Size(this, &FSourceControlViewportToolTips::GetToolTipSize)
+			[
+				SAssignNew(ToolTipWidget, SToolTip)
+				.Font(FCoreStyle::Get().GetFontStyle("ToolTip.Font"))
+				.Text_Lambda([this]() { return ToolTipText; })
+				.Visibility(EVisibility::Hidden)
+			];
+
+			ViewportWidget = LevelViewport;
+			ViewportWidget->AddOverlayWidget(CanvasWidget.ToSharedRef());
+		}
+	}
+}
+
+void FSourceControlViewportToolTips::RemoveCanvas()
+{
+	if (ViewportWidget.IsValid())
+	{
+		ViewportWidget->RemoveOverlayWidget(CanvasWidget.ToSharedRef());
+	}
+
+	CanvasWidget.Reset();
+	ToolTipWidget.Reset();
+	ViewportWidget.Reset();
+}
+
+FVector2D FSourceControlViewportToolTips::GetToolTipPosition() const
+{
+	FVector2D Result = FVector2D::Zero();
+
+	// Taken from SlateUser.cpp
+	static const FVector2f TooltipOffsetFromMouse(12.0f, 8.0f);
+	static const FVector2f TooltipOffsetFromForceField(4.0f, 3.0f);
+
+	if (ViewportWidget.IsValid())
+	{
+		FIntPoint ViewportOrigin;
+		FIntPoint ViewportSize;
+		ViewportWidget->GetViewportClient()->GetViewportDimensions(ViewportOrigin, ViewportSize);
+
+		int32 MouseX = 0;
+		int32 MouseY = 0;
+		if (FViewport* Viewport = ViewportWidget->GetViewportClient()->Viewport)
+		{
+			MouseX = Viewport->GetMouseX();
+			MouseY = Viewport->GetMouseY();
+		}
+
+		FVector2D ToolTipPos(MouseX + TooltipOffsetFromMouse.X, MouseY + TooltipOffsetFromMouse.Y);
+		FVector2D ToolTipSize = GetToolTipSize();
+	
+		if (ToolTipPos.X + ToolTipSize.X + TooltipOffsetFromForceField.X > ViewportSize.X)
+		{
+			// Flip tooltip left so we don't overflow on the right.
+			ToolTipPos.X = MouseX - TooltipOffsetFromMouse.X - ToolTipSize.X;
+		}
+		if (ToolTipPos.Y + ToolTipSize.Y + TooltipOffsetFromForceField.Y > ViewportSize.Y)
+		{
+			// Flip tooltip up so we don't overflow on the bottom.
+			ToolTipPos.Y = MouseY - TooltipOffsetFromMouse.Y - ToolTipSize.Y;
+		}
+
+		Result = ToolTipPos;
+	}
+
+	return Result;
+}
+
+FVector2D FSourceControlViewportToolTips::GetToolTipSize() const
+{
+	FVector2D Result = FVector2D::Zero();
+	
+	if (ToolTipWidget.IsValid())
+	{
+		FSlateFontInfo FontInfo = FCoreStyle::Get().GetFontStyle("ToolTip.Font");
+
+		FVector2D TextSize = FSlateApplication::Get().GetRenderer()->GetFontMeasureService()->Measure(ToolTipText, FontInfo);
+		FVector2D TextMargin = FVector2D(2 * 11.f, 2 * 11.f);
+
+		Result = TextSize + TextMargin;
+	}
+
+	return Result;
+}
+
+#undef LOCTEXT_NAMESPACE
\ No newline at end of file
diff --git a/Engine/Source/Developer/SourceControlViewport/Private/SourceControlViewportToolTips.h b/Engine/Source/Developer/SourceControlViewport/Private/SourceControlViewportToolTips.h
new file mode 100644
--- /dev/null
+++ b/Engine/Source/Developer/SourceControlViewport/Private/SourceControlViewportToolTips.h
@@ -0,0 +1,62 @@
+// Copyright Epic Games, Inc. All Rights Reserved.
+
+#pragma once
+
+#include "CoreMinimal.h"
+#include "Containers/Ticker.h"
+#include "UObject/WeakObjectPtr.h"
+
+class AActor;
+class SLevelViewport;
+class SCanvas;
+class SToolTip;
+
+// Adds a tooltip widget to the viewport.
+class FSourceControlViewportToolTips : public TSharedFromThis<FSourceControlViewportToolTips, ESPMode::ThreadSafe>
+{
+public:
+	FSourceControlViewportToolTips();
+	~FSourceControlViewportToolTips();
+
+public:
+	void Init();
+	bool Tick(float DeltaTime);
+
+private:
+	void UpdateCanvas(float DeltaTime);
+	void UpdateToolTip();
+
+	void InsertCanvas();
+	void RemoveCanvas();
+
+	FVector2D GetToolTipSize() const;
+	FVector2D GetToolTipPosition() const;
+
+private:
+	/* The ticker handle */
+	FTSTicker::FDelegateHandle TickHandle;
+
+	/* The canvas widget */
+	TSharedPtr<SCanvas> CanvasWidget;
+
+	/* The tooltip widget that'll be positioned on the canvas */
+	TSharedPtr<SToolTip> ToolTipWidget;
+
+	/* The viewport in which the canvas widget is added */
+	TSharedPtr<SLevelViewport> ViewportWidget;
+
+	/* The actor in the viewport that's currently hovered */
+	TWeakObjectPtr<AActor> Actor;
+
+	/* The tooltip text determined for the actor */
+	FText ToolTipText;
+
+	/* The previous mouse coordinates */
+	int32 ActorMouseX = 0;
+	int32 ActorMouseY = 0;
+
+	/* The time at which the actor became relevant */
+	double ActorTime = 0;
+	/* The time since the actor became relevant */
+	float DelayTime = 0;
+};
\ No newline at end of file
