Skein / ViewportStatus: Disable tooltip when actor is selected.

diff --git a/Engine/Source/Developer/SourceControlViewport/Private/SourceControlViewportToolTips.cpp b/Engine/Source/Developer/SourceControlViewport/Private/SourceControlViewportToolTips.cpp
--- a/Engine/Source/Developer/SourceControlViewport/Private/SourceControlViewportToolTips.cpp
+++ b/Engine/Source/Developer/SourceControlViewport/Private/SourceControlViewportToolTips.cpp
@@ -123,7 +123,7 @@ void FSourceControlViewportToolTips::UpdateToolTip()
 {
 	ToolTipText = FText::GetEmpty();
 
-	if (Actor.IsValid())
+	if (Actor.IsValid() && !Actor->IsSelected())
 	{
 		if (UPackage* Package = Actor->GetPackage())
 		{
