Skein / ViewportStatus: Integrate editor team work for Viewport Status Highlighting with UCS work.

diff --git a/Engine/Plugins/Experimental/TypedElementsDataStorage/Source/TypedElementsDataStorageRevisionControl/Source/Private/RevisionControlProcessors.cpp b/Engine/Plugins/Experimental/TypedElementsDataStorage/Source/TypedElementsDataStorageRevisionControl/Source/Private/RevisionControlProcessors.cpp
--- a/Engine/Plugins/Experimental/TypedElementsDataStorage/Source/TypedElementsDataStorageRevisionControl/Source/Private/RevisionControlProcessors.cpp
+++ b/Engine/Plugins/Experimental/TypedElementsDataStorage/Source/TypedElementsDataStorageRevisionControl/Source/Private/RevisionControlProcessors.cpp
@@ -18,56 +18,104 @@
 
 extern TYPEDELEMENTSDATASTORAGE_API FAutoConsoleVariableRef CVarAutoPopulateState;
 
-static bool GEnableOutlines = false;
-static FAutoConsoleVariableRef CVarEnableOutlines(
-	TEXT("TEDS.RevisionControl.UseOutlines"),
-	GEnableOutlines,
-	TEXT("Use revision control status outlines in the viewport")
-);
-
-static uint8 DetermineOutlineColorIndex(const TypedElementDataStorage::ICommonQueryContext& Context)
+static bool gEnableOverlays = true;
+TAutoConsoleVariable<bool> CVarEnableOverlays(
+	TEXT("SourceControl.Overlays.Enable"),
+	gEnableOverlays,
+	TEXT("Enables overlays."),
+	ECVF_Default);
+
+static bool gEnableOverlayCheckedOutByOtherUser = true;
+TAutoConsoleVariable<bool> CVarEnableOverlayCheckedOutByOtherUser(
+	TEXT("SourceControl.Overlays.CheckedOutByOtherUser.Enable"),
+	gEnableOverlayCheckedOutByOtherUser,
+	TEXT("Enables overlays for files that are checked out by another user."),
+	ECVF_Default);
+
+static bool gEnableOverlayNotAtHeadRevision = true;
+TAutoConsoleVariable<bool> CVarEnableOverlayNotAtHeadRevision(
+	TEXT("SourceControl.Overlays.NotAtHeadRevision.Enable"),
+	gEnableOverlayNotAtHeadRevision,
+	TEXT("Enables overlays for files that are not at the latest revision."),
+	ECVF_Default);
+
+static bool gEnableOverlayCheckedOut = false;
+TAutoConsoleVariable<bool> CVarEnableOverlayCheckedOut(
+	TEXT("SourceControl.Overlays.CheckedOut.Enable"),
+	gEnableOverlayCheckedOut,
+	TEXT("Enables overlays for files that are checked out by user."),
+	ECVF_Default);
+
+static bool gEnableOverlayOpenForAdd = false;
+TAutoConsoleVariable<bool> CVarEnableOverlayOpenForAdd(
+	TEXT("SourceControl.Overlays.OpenForAdd.Enable"),
+	gEnableOverlayOpenForAdd,
+	TEXT("Enables overlays for files that are newly added."),
+	ECVF_Default);
+
+static int32 gOverlayAlpha = 217;
+TAutoConsoleVariable<int32> CVarOverlayAlpha(
+	TEXT("SourceControl.Overlays.Alpha"),
+	gOverlayAlpha,
+	TEXT("Configures overlay opacity."),
+	ECVF_Default);
+
+static FColor DetermineOverlayColor(const TypedElementDataStorage::IQueryContext& ObjectContext, const TypedElementDataStorage::ICommonQueryContext& SCCContext)
 {
-	if (GEnableOutlines)
-	{
-		// TODO: Get this information from TEDS instead of hardcoded pointing into UEditorStyleSettings::AdditionalSelectionColors?
-		constexpr uint8 BasicSelectionColorCount = 2;
-		constexpr uint8 IndexBlue = BasicSelectionColorCount + 0;
-		constexpr uint8 IndexPurple = BasicSelectionColorCount + 1;
-		constexpr uint8 IndexPink = BasicSelectionColorCount + 2;
-		constexpr uint8 IndexRed = BasicSelectionColorCount + 3;
-		constexpr uint8 IndexYellow = BasicSelectionColorCount + 4;
-		constexpr uint8 IndexGreen = BasicSelectionColorCount + 5;
+	check(IsInGameThread());
 
+	bool bEnabled = CVarEnableOverlays.GetValueOnGameThread();
+	bool bSelected = ObjectContext.HasColumn<FTypedElementSelectionColumn>();
+	if (bEnabled && !bSelected)
+	{
 		// Check if the package is outdated because there is a newer version available.
-		if (Context.HasColumn<FSCCNotCurrentTag>())
+		if (SCCContext.HasColumn<FSCCNotCurrentTag>())
 		{
-			return IndexYellow;
+			if (CVarEnableOverlayNotAtHeadRevision.GetValueOnGameThread())
+			{
+				// Yellow.
+				return FColor(255, 255, 61, CVarOverlayAlpha.GetValueOnGameThread());
+			}
 		}
 
 		// Check if the package is locked by someone else.
-		if (Context.HasColumn<FSCCExternallyLockedColumn>())
+		if (SCCContext.HasColumn<FSCCExternallyLockedColumn>())
 		{
-			return IndexRed;
+			if (CVarEnableOverlayCheckedOutByOtherUser.GetValueOnGameThread())
+			{
+				// Red.
+				return FColor(255, 64, 64, CVarOverlayAlpha.GetValueOnGameThread());
+			}
 		}
 
 		// Check if the package is added locally.
-		if (Context.HasColumn<FSCCStatusColumn>())
+		if (SCCContext.HasColumn<FSCCStatusColumn>())
 		{
-			const FSCCStatusColumn* StatusColumn = Context.GetColumn<FSCCStatusColumn>();
-			if (StatusColumn->Modification == ESCCModification::Added)
+			if (CVarEnableOverlayOpenForAdd.GetValueOnGameThread())
 			{
-				return IndexGreen;
+				if (const FSCCStatusColumn* StatusColumn = SCCContext.GetColumn<FSCCStatusColumn>())
+				{
+					if (StatusColumn->Modification == ESCCModification::Added)
+					{
+						// Green.
+						return FColor(134, 194, 74, CVarOverlayAlpha.GetValueOnGameThread());
+					}
+				}
 			}
 		}
 
 		// Check if the package is locked by self.
-		if (Context.HasColumn<FSCCLockedTag>())
+		if (SCCContext.HasColumn<FSCCLockedTag>())
 		{
-			return IndexBlue;
+			if (CVarEnableOverlayCheckedOut.GetValueOnGameThread())
+			{
+				// Blue.
+				return FColor(0, 112, 224, CVarOverlayAlpha.GetValueOnGameThread());
+			}
 		}
 	}
 
-	return 0; // Default outline color.
+	return FColor(ForceInitToZero);
 }
 
 void UTypedElementRevisionControlFactory::RegisterTables(ITypedElementDataStorageInterface& DataStorage)
@@ -87,19 +135,29 @@ void UTypedElementRevisionControlFactory::RegisterQueries(ITypedElementDataStora
 	TypedElementQueryHandle ObjectToSCCQuery = DataStorage.RegisterQuery(
 		Select()
 			.ReadOnly<FTypedElementPackagePathColumn>()
+			.ReadOnly<FSCCStatusColumn>(EOptional::Yes)
 		.Compile());
 
 	DataStorage.RegisterQuery(
 		Select(
 			TEXT("Change selection outline colors based on SCC status"),
 			// This is in PrePhysics because the outline->actor query is in DuringPhysics and contexts don't flush changes between tick groups
-			FProcessor(DSI::EQueryTickPhase::PrePhysics, DataStorage.GetQueryTickGroupName(DSI::EQueryTickGroups::SyncExternalToDataStorage)),
+			FProcessor(DSI::EQueryTickPhase::PrePhysics, DataStorage.GetQueryTickGroupName(DSI::EQueryTickGroups::SyncExternalToDataStorage))
+				.ForceToGameThread(true),
 			[](DSI::IQueryContext& Context, TypedElementRowHandle ObjectRow, const FTypedElementPackageReference& PackageReference)
 			{
 				Context.RunSubquery(0, PackageReference.Row, CreateSubqueryCallbackBinding(
 					[&Context, &ObjectRow](DSI::ISubqueryContext& SubQueryContext)
 					{
-						Context.AddColumn<FTypedElementViewportOutlineColorColumn>(ObjectRow, { .SelectionOutlineColorIndex = DetermineOutlineColorIndex(SubQueryContext) });
+						FColor Color = DetermineOverlayColor(Context, SubQueryContext);
+						if (Color.Bits != 0)
+						{
+							Context.AddColumn<FTypedElementViewportOverlayColorColumn>(ObjectRow, { .OverlayColor = Color });
+						}
+						else
+						{
+							Context.RemoveColumns<FTypedElementViewportOverlayColorColumn>(ObjectRow);
+						}
 						Context.AddColumns<FTypedElementSyncBackToWorldTag>(ObjectRow);
 					})
 				);
diff --git a/Engine/Source/Developer/SourceControl/Private/RevisionControlStyle/RevisionControlStyle.cpp b/Engine/Source/Developer/SourceControl/Private/RevisionControlStyle/RevisionControlStyle.cpp
--- a/Engine/Source/Developer/SourceControl/Private/RevisionControlStyle/RevisionControlStyle.cpp
+++ b/Engine/Source/Developer/SourceControl/Private/RevisionControlStyle/RevisionControlStyle.cpp
@@ -102,7 +102,7 @@ FDefaultRevisionControlStyle::FDefaultRevisionControlStyle() : FSlateStyleSet(St
 	Set("RevisionControl.Locked", new IMAGE_BRUSH_SVG("Starship/Common/lock", CoreStyleConstants::Icon16x16));
 	Set("RevisionControl.Unlocked", new IMAGE_BRUSH_SVG("Starship/Common/lock-unlocked", CoreStyleConstants::Icon16x16));
 
-	Set("RevisionControl.CheckedOut", new IMAGE_BRUSH_SVG("Starship/SourceControl/SCC_CheckedOut", CoreStyleConstants::Icon16x16, FStyleColors::Error));
+	Set("RevisionControl.CheckedOut", new IMAGE_BRUSH_SVG("Starship/SourceControl/SCC_CheckedOut", CoreStyleConstants::Icon16x16, FStyleColors::AccentBlue));
 	
 	Set("RevisionControl.OpenForAdd", new IMAGE_BRUSH_SVG("Starship/SourceControl/RC_MarkedForAdd", CoreStyleConstants::Icon16x16, FStyleColors::AccentGreen));
 	Set("RevisionControl.MarkedForDelete", new IMAGE_BRUSH_SVG("Starship/SourceControl/SCC_MarkedForDelete", CoreStyleConstants::Icon16x16, FStyleColors::Error));
@@ -135,6 +135,11 @@ FDefaultRevisionControlStyle::FDefaultRevisionControlStyle() : FSlateStyleSet(St
 	Set("RevisionControl.StatusBar.HasLocalChanges", new IMAGE_BRUSH_SVG("Starship/SourceControl/RC_StatusLocalUpload", CoreStyleConstants::Icon16x16, FStyleColors::AccentBlue));
 	Set("RevisionControl.StatusBar.Conflicted", new IMAGE_BRUSH_SVG("Starship/SourceControl/RC_Conflicted", CoreStyleConstants::Icon16x16, FStyleColors::Warning));
 	Set("RevisionControl.StatusBar.Promote", new IMAGE_BRUSH_SVG("Starship/SourceControl/RC_Promote", CoreStyleConstants::Icon16x16));
+
+	Set("RevisionControl.ShowMenu.CheckedOut", new IMAGE_BRUSH_SVG("Starship/SourceControl/SCC_CheckedOut", CoreStyleConstants::Icon16x16));
+	Set("RevisionControl.ShowMenu.OpenForAdd", new IMAGE_BRUSH_SVG("Starship/SourceControl/RC_MarkedForAdd", CoreStyleConstants::Icon16x16));
+	Set("RevisionControl.ShowMenu.CheckedOutByOtherUser", new IMAGE_BRUSH_SVG("Starship/SourceControl/RC_CheckedOther", CoreStyleConstants::Icon16x16));
+	Set("RevisionControl.ShowMenu.NotAtHeadRevision", new IMAGE_BRUSH_SVG("Starship/SourceControl/RC_NewerVersion", CoreStyleConstants::Icon16x16));
 	
 	Set("RevisionControl.Promote.Large", new IMAGE_BRUSH_SVG("Starship/SourceControl/RC_Promote_Large", CoreStyleConstants::Icon32x32));
 
diff --git a/Engine/Source/Developer/SourceControlViewport/Private/SourceControlViewportMenu.cpp b/Engine/Source/Developer/SourceControlViewport/Private/SourceControlViewportMenu.cpp
new file mode 100644
--- /dev/null
+++ b/Engine/Source/Developer/SourceControlViewport/Private/SourceControlViewportMenu.cpp
@@ -0,0 +1,282 @@
+// Copyright Epic Games, Inc. All Rights Reserved.
+
+#include "SourceControlViewportMenu.h"
+#include "RevisionControlStyle/RevisionControlStyle.h"
+#include "Widgets/Input/SSpinBox.h"
+#include "EngineAnalytics.h"
+#include "LevelEditorMenuContext.h"
+#include "LevelEditorViewport.h"
+#include "ToolMenus.h"
+
+#define LOCTEXT_NAMESPACE "SourceControlViewportMenu"
+
+static const FName MenuName("LevelEditor.LevelViewportToolbar.Show");
+static const FName SectionName("LevelViewportEditorShow");
+static const FName SubMenuName("ShowRevisionControlMenu");
+
+FSourceControlViewportMenu::FSourceControlViewportMenu()
+{
+}
+
+FSourceControlViewportMenu::~FSourceControlViewportMenu()
+{
+	RemoveViewportMenu();
+}
+
+void FSourceControlViewportMenu::Init()
+{
+}
+
+void FSourceControlViewportMenu::SetEnabled(bool bInEnabled)
+{
+	if (bInEnabled)
+	{
+		InsertViewportMenu();
+	}
+	else
+	{
+		RemoveViewportMenu();
+	}
+}
+
+void FSourceControlViewportMenu::InsertViewportMenu()
+{
+	if (UToolMenus* ToolMenus = UToolMenus::TryGet())
+	{
+		UToolMenu* Menu = ToolMenus->ExtendMenu(MenuName);
+		if (Menu != nullptr)
+		{
+			Menu->AddDynamicSection(SectionName,
+				FNewToolMenuDelegate::CreateSP(this, &FSourceControlViewportMenu::PopulateViewportMenu)
+			);
+		}
+	}
+}
+
+void FSourceControlViewportMenu::PopulateViewportMenu(UToolMenu* InMenu)
+{
+	check(InMenu);
+
+	InMenu->AddDynamicSection(SectionName, FNewToolMenuDelegate::CreateLambda(
+		[this](UToolMenu* InMenu)
+		{
+			ULevelViewportToolBarContext* Context = InMenu->FindContext<ULevelViewportToolBarContext>();
+			if (!Context)
+			{
+				return;
+			}
+
+			FLevelEditorViewportClient* ViewportClient = Context->GetLevelViewportClient();
+			if (!ViewportClient || !ViewportClient->IsPerspective())
+			{
+				return;
+			}
+
+			OpacityWidget = SNew(SSpinBox<uint8>)
+				.ClearKeyboardFocusOnCommit(true)
+				.OnValueChanged_Lambda(
+					[this, ViewportClient](uint8 InNewValue)
+					{
+						SetOpacityValue(ViewportClient, InNewValue);
+					}
+				)
+				.OnValueCommitted_Lambda(
+					[this, ViewportClient](uint8 InNewValue, ETextCommit::Type InCommitType)
+					{
+						SetOpacityValue(ViewportClient, InNewValue);
+					}
+				)
+				.Value_Lambda(
+					[this, ViewportClient]()
+					{
+						return GetOpacityValue(ViewportClient);
+					}
+				)
+				.MinValue(0)
+				.MinSliderValue(0)
+				.MaxValue(100)
+				.MaxSliderValue(100);
+
+			FToolMenuSection& RevisionControlSection = InMenu->FindOrAddSection(TEXT("LevelViewportEditorShow"));
+			RevisionControlSection.AddDynamicEntry(NAME_None, FNewToolMenuSectionDelegate::CreateLambda(
+				[this, ViewportClient](FToolMenuSection& InSection)
+				{
+					InSection.AddSubMenu(
+						SubMenuName,
+						LOCTEXT("RevisionControlSubMenu", "Revision Control"),
+						LOCTEXT("RevisionControlSubMenu_ToolTip", "Toggle revision control viewport options on or off."),
+						FNewToolMenuDelegate::CreateLambda(
+							[this, ViewportClient](UToolMenu* InSubMenu)
+							{
+								FToolMenuSection& DefaultSection = InSubMenu->AddSection(NAME_None, LOCTEXT("RevisionControlSectionStatus", "Status Highlights"));
+
+								DefaultSection.AddMenuEntry(
+									NAME_None,
+									LOCTEXT("ShowAll", "Show All"),
+									LOCTEXT("ShowAll_ToolTip", "Enable highlighting for all statuses"),
+									FSlateIcon(),
+									FUIAction(
+										FExecuteAction::CreateSP(this, &FSourceControlViewportMenu::ShowAll, ViewportClient)
+									),
+									EUserInterfaceActionType::Button
+								);
+
+								DefaultSection.AddMenuEntry(
+									NAME_None,
+									LOCTEXT("HideAll", "Hide All"),
+									LOCTEXT("HideAll_ToolTip", "Disable highlighting for all statuses"),
+									FSlateIcon(),
+									FUIAction(
+										FExecuteAction::CreateSP(this, &FSourceControlViewportMenu::HideAll, ViewportClient)
+									),
+									EUserInterfaceActionType::Button
+								);
+
+								DefaultSection.AddSeparator(NAME_None);
+
+								DefaultSection.AddMenuEntry(
+									NAME_None,
+									LOCTEXT("HighlightCheckedOutByOtherUser", "Checked Out by Others"),
+									LOCTEXT("HighlightCheckedOutByOtherUser_ToolTip", "Highlight objects that are checked out by someone else."),
+									FSlateIcon(FRevisionControlStyleManager::GetStyleSetName(), "RevisionControl.ShowMenu.CheckedOutByOtherUser"),
+									FUIAction(
+										FExecuteAction::CreateSP(this, &FSourceControlViewportMenu::ToggleHighlight, ViewportClient, ESourceControlStatus::CheckedOutByOtherUser),
+										FCanExecuteAction(),
+										FIsActionChecked::CreateSP(this, &FSourceControlViewportMenu::IsHighlighted, ViewportClient, ESourceControlStatus::CheckedOutByOtherUser)
+									),
+									EUserInterfaceActionType::ToggleButton
+								);
+
+								DefaultSection.AddMenuEntry(
+									NAME_None,
+									LOCTEXT("HighlightNotAtHeadRevision", "Out of Date"),
+									LOCTEXT("HighlightNotAtHeadRevision_ToolTip", "Highlight objects that are not at the latest revision."),
+									FSlateIcon(FRevisionControlStyleManager::GetStyleSetName(), "RevisionControl.ShowMenu.NotAtHeadRevision"),
+									FUIAction(
+										FExecuteAction::CreateSP(this, &FSourceControlViewportMenu::ToggleHighlight, ViewportClient, ESourceControlStatus::NotAtHeadRevision),
+										FCanExecuteAction(),
+										FIsActionChecked::CreateSP(this, &FSourceControlViewportMenu::IsHighlighted, ViewportClient, ESourceControlStatus::NotAtHeadRevision)
+									),
+									EUserInterfaceActionType::ToggleButton
+								);
+
+								DefaultSection.AddMenuEntry(
+									NAME_None,
+									LOCTEXT("HighlightCheckedOut", "Checked Out by Me"),
+									LOCTEXT("HighlightCheckedOut_ToolTip", "Highlight objects that are checked out by me."),
+									FSlateIcon(FRevisionControlStyleManager::GetStyleSetName(), "RevisionControl.ShowMenu.CheckedOut"),
+									FUIAction(
+										FExecuteAction::CreateSP(this, &FSourceControlViewportMenu::ToggleHighlight, ViewportClient, ESourceControlStatus::CheckedOut),
+										FCanExecuteAction(),
+										FIsActionChecked::CreateSP(this, &FSourceControlViewportMenu::IsHighlighted, ViewportClient, ESourceControlStatus::CheckedOut)
+									),
+									EUserInterfaceActionType::ToggleButton
+								);
+
+								DefaultSection.AddMenuEntry(
+									NAME_None,
+									LOCTEXT("HighlightOpenForAdd", "Newly Added"),
+									LOCTEXT("HighlightOpenForAdd_ToolTip", "Highlight objects that have been added by me."),
+									FSlateIcon(FRevisionControlStyleManager::GetStyleSetName(), "RevisionControl.ShowMenu.OpenForAdd"),
+									FUIAction(
+										FExecuteAction::CreateSP(this, &FSourceControlViewportMenu::ToggleHighlight, ViewportClient, ESourceControlStatus::OpenForAdd),
+										FCanExecuteAction(),
+										FIsActionChecked::CreateSP(this, &FSourceControlViewportMenu::IsHighlighted, ViewportClient, ESourceControlStatus::OpenForAdd)
+									),
+									EUserInterfaceActionType::ToggleButton
+								);
+
+								DefaultSection.AddEntry(FToolMenuEntry::InitWidget(
+									NAME_None,
+									OpacityWidget.ToSharedRef(),
+									LOCTEXT("Opacity", "Opacity")
+									)
+								);
+							}
+						),
+						false,
+						FSlateIcon(FAppStyle::Get().GetStyleSetName(), "ShowFlagsMenu.SubMenu.RevisionControl")
+					);
+				}));
+		}
+	)
+	);
+}
+
+void FSourceControlViewportMenu::RemoveViewportMenu()
+{
+	if (UToolMenus* ToolMenus = UToolMenus::TryGet())
+	{
+		UToolMenu* Menu = ToolMenus->ExtendMenu(MenuName);
+		if (Menu != nullptr)
+		{
+			Menu->RemoveSection(SectionName);
+		}
+	}
+}
+
+void FSourceControlViewportMenu::ShowAll(FLevelEditorViewportClient* ViewportClient)
+{
+	ensure(ViewportClient);
+
+	SourceControlViewportUtils::SetFeedbackEnabled(ViewportClient, ESourceControlStatus::CheckedOutByOtherUser, /*bEnabled=*/true);
+	SourceControlViewportUtils::SetFeedbackEnabled(ViewportClient, ESourceControlStatus::NotAtHeadRevision, /*bEnabled=*/true);
+	SourceControlViewportUtils::SetFeedbackEnabled(ViewportClient, ESourceControlStatus::CheckedOut, /*bEnabled=*/true);
+	SourceControlViewportUtils::SetFeedbackEnabled(ViewportClient, ESourceControlStatus::OpenForAdd, /*bEnabled=*/true);
+	
+	RecordToggleEvent(TEXT("All"), /*bEnabled=*/true);
+}
+
+void FSourceControlViewportMenu::HideAll(FLevelEditorViewportClient* ViewportClient)
+{
+	ensure(ViewportClient);
+
+	SourceControlViewportUtils::SetFeedbackEnabled(ViewportClient, ESourceControlStatus::CheckedOutByOtherUser, /*bEnabled=*/false);
+	SourceControlViewportUtils::SetFeedbackEnabled(ViewportClient, ESourceControlStatus::NotAtHeadRevision, /*bEnabled=*/false);
+	SourceControlViewportUtils::SetFeedbackEnabled(ViewportClient, ESourceControlStatus::CheckedOut, /*bEnabled=*/false);
+	SourceControlViewportUtils::SetFeedbackEnabled(ViewportClient, ESourceControlStatus::OpenForAdd, /*bEnabled=*/false);
+
+	RecordToggleEvent(TEXT("All"), /*bEnabled=*/false);
+}
+
+void FSourceControlViewportMenu::ToggleHighlight(FLevelEditorViewportClient* ViewportClient, ESourceControlStatus Status)
+{
+	ensure(ViewportClient);
+
+	bool bOld = SourceControlViewportUtils::GetFeedbackEnabled(ViewportClient, Status);
+	bool bNew = bOld ? false : true;
+	SourceControlViewportUtils::SetFeedbackEnabled(ViewportClient, Status, bNew);
+
+	FString EnumValueWithoutType = UEnum::GetValueAsString(Status)
+		.Replace(TEXT("ESourceControlStatus::"), TEXT(""));
+	RecordToggleEvent(EnumValueWithoutType, bNew);
+}
+
+bool FSourceControlViewportMenu::IsHighlighted(FLevelEditorViewportClient* ViewportClient, ESourceControlStatus Status) const
+{
+	ensure(ViewportClient);
+
+	return SourceControlViewportUtils::GetFeedbackEnabled(ViewportClient, Status);
+}
+
+void FSourceControlViewportMenu::SetOpacityValue(FLevelEditorViewportClient* ViewportClient, uint8 InNewValue)
+{
+	SourceControlViewportUtils::SetFeedbackOpacity(ViewportClient, InNewValue);
+}
+
+uint8 FSourceControlViewportMenu::GetOpacityValue(FLevelEditorViewportClient* ViewportClient) const
+{
+	return SourceControlViewportUtils::GetFeedbackOpacity(ViewportClient);
+}
+
+void FSourceControlViewportMenu::RecordToggleEvent(const FString& Param, bool bEnabled) const
+{
+	if (FEngineAnalytics::IsAvailable())
+	{
+		FEngineAnalytics::GetProvider().RecordEvent(
+			TEXT("Editor.Usage.SourceControl.Settings"), Param, bEnabled ? TEXT("True") : TEXT("False")
+		);
+	}
+}
+
+#undef LOCTEXT_NAMESPACE
\ No newline at end of file
diff --git a/Engine/Source/Developer/SourceControlViewport/Private/SourceControlViewportMenu.h b/Engine/Source/Developer/SourceControlViewport/Private/SourceControlViewportMenu.h
new file mode 100644
--- /dev/null
+++ b/Engine/Source/Developer/SourceControlViewport/Private/SourceControlViewportMenu.h
@@ -0,0 +1,43 @@
+// Copyright Epic Games, Inc. All Rights Reserved.
+
+#pragma once
+
+#include "CoreMinimal.h"
+#include "SourceControlViewportUtils.h"
+
+class FLevelEditorViewportClient;
+class UToolMenu;
+class SWidget;
+
+// Adds an options menu to the Viewport's SHOW pill.
+class FSourceControlViewportMenu : public TSharedFromThis<FSourceControlViewportMenu, ESPMode::ThreadSafe>
+{
+public:
+	FSourceControlViewportMenu();
+	~FSourceControlViewportMenu();
+
+public:
+	void Init();
+	void SetEnabled(bool bInEnabled);
+
+private:
+	void InsertViewportMenu();
+	void PopulateViewportMenu(UToolMenu* InMenu);
+	void RemoveViewportMenu();
+
+private:
+	void ShowAll(FLevelEditorViewportClient* ViewportClient);
+	void HideAll(FLevelEditorViewportClient* ViewportClient);
+
+	void ToggleHighlight(FLevelEditorViewportClient* ViewportClient, ESourceControlStatus Status);
+	bool IsHighlighted(FLevelEditorViewportClient* ViewportClient, ESourceControlStatus Status) const;
+
+	void SetOpacityValue(FLevelEditorViewportClient* ViewportClient, uint8 NewValue);
+	uint8 GetOpacityValue(FLevelEditorViewportClient* ViewportClient) const;
+
+private:
+	void RecordToggleEvent(const FString& Param, bool bEnabled) const;
+
+private:
+	TSharedPtr<SWidget> OpacityWidget;
+};
\ No newline at end of file
diff --git a/Engine/Source/Developer/SourceControlViewport/Private/SourceControlViewportModule.cpp b/Engine/Source/Developer/SourceControlViewport/Private/SourceControlViewportModule.cpp
--- a/Engine/Source/Developer/SourceControlViewport/Private/SourceControlViewportModule.cpp
+++ b/Engine/Source/Developer/SourceControlViewport/Private/SourceControlViewportModule.cpp
@@ -1,21 +1,79 @@
 // Copyright Epic Games, Inc. All Rights Reserved.
 
 #include "SourceControlViewportModule.h"
-#include "SourceControlViewportOutlineMenu.h"
+#include "SourceControlViewportMenu.h"
 #include "SourceControlViewportToolTips.h"
+#include "HAL/IConsoleManager.h"
+
+TAutoConsoleVariable<bool> CVarSourceControlEnableViewportStatus(
+	TEXT("SourceControl.ViewportStatus.Enable"),
+	false,
+	TEXT("Enables source control viewport status features."),
+	ECVF_Default);
 
 void FSourceControlViewportModule::StartupModule()
 {
-	ViewportOutlineMenu = MakeShared<FSourceControlViewportOutlineMenu>();
-	ViewportOutlineMenu->Init();
+	ViewportMenu = MakeShared<FSourceControlViewportMenu>();
+	ViewportMenu->Init();
 	ViewportToolTips = MakeShared<FSourceControlViewportToolTips>();
 	ViewportToolTips->Init();
+
+	if (IConsoleVariable* CVar = IConsoleManager::Get().FindConsoleVariable(TEXT("SourceControl.ViewportStatus.Enable")); ensure(CVar))
+	{
+		CVar->OnChangedDelegate().AddRaw(this, &FSourceControlViewportModule::HandleCVarChanged);
+	}
+
+	if (IConsoleVariable* CVar = IConsoleManager::Get().FindConsoleVariable(TEXT("TEDS.RevisionControl.AutoPopulateState")); ensure(CVar))
+	{
+		CVar->OnChangedDelegate().AddRaw(this, &FSourceControlViewportModule::HandleCVarChanged);
+	}
+
+	UpdateSettings();
 }
 
 void FSourceControlViewportModule::ShutdownModule()
 {
+	if (IConsoleVariable* CVar = IConsoleManager::Get().FindConsoleVariable(TEXT("TEDS.RevisionControl.AutoPopulateState")); ensure(CVar))
+	{
+		CVar->OnChangedDelegate().RemoveAll(this);
+	}
+
+	if (IConsoleVariable* CVar = IConsoleManager::Get().FindConsoleVariable(TEXT("SourceControl.ViewportStatus.Enable")); ensure(CVar))
+	{
+		CVar->OnChangedDelegate().RemoveAll(this);
+	}
+
 	ViewportToolTips.Reset();
-	ViewportOutlineMenu.Reset();
+	ViewportMenu.Reset();
+}
+
+void FSourceControlViewportModule::HandleCVarChanged(IConsoleVariable* Variable)
+{
+	UpdateSettings();
+}
+
+void FSourceControlViewportModule::UpdateSettings()
+{
+	// Are both CVar's enabled?
+	bool bEnabled = true;
+
+	if (IConsoleVariable* CVar = IConsoleManager::Get().FindConsoleVariable(TEXT("SourceControl.ViewportStatus.Enable")); ensure(CVar))
+	{
+		bEnabled &= CVar->GetBool();
+	}
+	if (IConsoleVariable* CVar = IConsoleManager::Get().FindConsoleVariable(TEXT("TEDS.RevisionControl.AutoPopulateState")); ensure(CVar))
+	{
+		bEnabled &= CVar->GetBool();
+	}
+
+	// Propagate resulting value.
+	ViewportMenu->SetEnabled(bEnabled);
+	ViewportToolTips->SetEnabled(bEnabled);
+
+	if (IConsoleVariable* CVar = IConsoleManager::Get().FindConsoleVariable(TEXT("SourceControl.Overlays.Enable")); ensure(CVar))
+	{
+		CVar->Set(bEnabled);
+	}
 }
 
 IMPLEMENT_MODULE( FSourceControlViewportModule, SourceControlViewport );
\ No newline at end of file
diff --git a/Engine/Source/Developer/SourceControlViewport/Private/SourceControlViewportModule.h b/Engine/Source/Developer/SourceControlViewport/Private/SourceControlViewportModule.h
--- a/Engine/Source/Developer/SourceControlViewport/Private/SourceControlViewportModule.h
+++ b/Engine/Source/Developer/SourceControlViewport/Private/SourceControlViewportModule.h
@@ -5,7 +5,7 @@
 #include "CoreMinimal.h"
 #include "Modules/ModuleManager.h"
 
-class FSourceControlViewportOutlineMenu;
+class FSourceControlViewportMenu;
 class FSourceControlViewportToolTips;
 
 class FSourceControlViewportModule : public IModuleInterface
@@ -16,9 +16,9 @@ public:
 	virtual void ShutdownModule() override;
 
 	/**
-	 * Gets a reference to the source control outline menu module instance.
+	 * Gets a reference to the source control viewport module instance.
 	 *
-	 * @return A reference to the source control outline menu module.
+	 * @return A reference to the source control viewport module.
 	 */
 	static FSourceControlViewportModule& Get()
 	{
@@ -33,6 +33,10 @@ public:
 	}
 
 private:
-	TSharedPtr<FSourceControlViewportOutlineMenu> ViewportOutlineMenu;
+	void HandleCVarChanged(IConsoleVariable* Variable);
+	void UpdateSettings();
+
+private:
+	TSharedPtr<FSourceControlViewportMenu> ViewportMenu;
 	TSharedPtr<FSourceControlViewportToolTips> ViewportToolTips;
 };
diff --git a/Engine/Source/Developer/SourceControlViewport/Private/SourceControlViewportOutlineMenu.cpp b/Engine/Source/Developer/SourceControlViewport/Private/SourceControlViewportOutlineMenu.cpp
deleted file mode 100644
diff --git a/Engine/Source/Developer/SourceControlViewport/Private/SourceControlViewportOutlineMenu.h b/Engine/Source/Developer/SourceControlViewport/Private/SourceControlViewportOutlineMenu.h
deleted file mode 100644
diff --git a/Engine/Source/Developer/SourceControlViewport/Private/SourceControlViewportToolTips.cpp b/Engine/Source/Developer/SourceControlViewport/Private/SourceControlViewportToolTips.cpp
--- a/Engine/Source/Developer/SourceControlViewport/Private/SourceControlViewportToolTips.cpp
+++ b/Engine/Source/Developer/SourceControlViewport/Private/SourceControlViewportToolTips.cpp
@@ -6,7 +6,6 @@
 #include "Widgets/SToolTip.h"
 #include "Widgets/SCanvas.h"
 #include "Fonts/FontMeasure.h"
-#include "HAL/IConsoleManager.h"
 #include "EngineUtils.h"
 #include "LevelEditor.h"
 #include "SourceControlHelpers.h"
@@ -17,13 +16,6 @@
 
 #define LOCTEXT_NAMESPACE "SourceControlViewportToolTips"
 
-static bool bEnableViewportToolTips = false;
-TAutoConsoleVariable<bool> CVarSourceControlEnableViewportToolTips(
-	TEXT("SourceControl.ViewportToolTips.Enable"),
-	bEnableViewportToolTips,
-	TEXT("Enables source control tooltips in the viewport."),
-	ECVF_Default);
-
 FSourceControlViewportToolTips::FSourceControlViewportToolTips()
 {
 }
@@ -41,6 +33,11 @@ void FSourceControlViewportToolTips::Init()
 	);
 }
 
+void FSourceControlViewportToolTips::SetEnabled(bool bInEnabled)
+{
+	bEnabled = bInEnabled;
+}
+
 bool FSourceControlViewportToolTips::Tick(float DeltaTime)
 {
 	UpdateCanvas(DeltaTime);
@@ -62,11 +59,9 @@ void FSourceControlViewportToolTips::UpdateCanvas(float DeltaTime)
 		}
 	}
 
-	bool bEnabled = CVarSourceControlEnableViewportToolTips.GetValueOnGameThread();
-
 	if (CanvasWidget.IsValid() && !bEnabled)
 	{
-		// Remove the canvas if the CVar got disabled.
+		// Remove the canvas if the tooltips got disabled.
 		RemoveCanvas();
 	}
 
@@ -78,43 +73,46 @@ void FSourceControlViewportToolTips::UpdateCanvas(float DeltaTime)
 
 	if (CanvasWidget.IsValid() && ToolTipWidget.IsValid() && ViewportWidget.IsValid())
 	{
-		FViewport* Viewport = ViewportWidget->GetViewportClient()->Viewport;
-		check(Viewport);
-
-		int32 MouseX = Viewport->GetMouseX();
-		int32 MouseY = Viewport->GetMouseY();
-		if (ActorMouseX != MouseX || ActorMouseY != MouseY)
+		if (TSharedPtr<SLevelViewport> ViewportWidgetPtr = ViewportWidget.Pin())
 		{
-			Actor.Reset();
-			UpdateToolTip();
-
-			ActorMouseX = MouseX;
-			ActorMouseY = MouseY;
-		}
+			FViewport* Viewport = ViewportWidgetPtr->GetViewportClient()->Viewport;
+			check(Viewport);
 
-		HActor* ActorHitProxy = HitProxyCast<HActor>(Viewport->GetHitProxy(MouseX, MouseY));
-		if (ActorHitProxy == nullptr)
-		{
-			Actor.Reset();
-			UpdateToolTip();
-		}
-		if (ActorHitProxy != nullptr && ActorHitProxy->Actor != Actor)
-		{
-			Actor.Reset();
-			UpdateToolTip();
+			int32 MouseX = Viewport->GetMouseX();
+			int32 MouseY = Viewport->GetMouseY();
+			if (ActorMouseX != MouseX || ActorMouseY != MouseY)
+			{
+				Actor.Reset();
+				UpdateToolTip();
 
-			Actor = MakeWeakObjectPtr(ActorHitProxy->Actor);
-			ActorTime = FPlatformTime::Seconds();
-			DelayTime = 0;
-		}
+				ActorMouseX = MouseX;
+				ActorMouseY = MouseY;
+			}
 
-		if (Actor.IsValid())
-		{
-			DelayTime += DeltaTime;
-			if (DelayTime >= 0.5f)
+			HActor* ActorHitProxy = HitProxyCast<HActor>(Viewport->GetHitProxy(MouseX, MouseY));
+			if (ActorHitProxy == nullptr)
 			{
+				Actor.Reset();
 				UpdateToolTip();
 			}
+			if (ActorHitProxy != nullptr && ActorHitProxy->Actor != Actor)
+			{
+				Actor.Reset();
+				UpdateToolTip();
+
+				Actor = MakeWeakObjectPtr(ActorHitProxy->Actor);
+				ActorTime = FPlatformTime::Seconds();
+				DelayTime = 0;
+			}
+
+			if (Actor.IsValid())
+			{
+				DelayTime += DeltaTime;
+				if (DelayTime >= 0.5f)
+				{
+					UpdateToolTip();
+				}
+			}
 		}
 	}
 }
@@ -143,10 +141,18 @@ void FSourceControlViewportToolTips::UpdateToolTip()
 					bool bCheckedOut = State->IsCheckedOut();
 					bool bOpenForAdd = State->IsAdded();
 
-					bool bNotAtHeadRevisionEnabled = SourceControlViewportUtils::GetFeedbackEnabled(ViewportWidget->GetViewportClient().Get(), ESourceControlStatus::NotAtHeadRevision);
-					bool bCheckedOutByOtherUserEnabled = SourceControlViewportUtils::GetFeedbackEnabled(ViewportWidget->GetViewportClient().Get(), ESourceControlStatus::CheckedOutByOtherUser);
-					bool bCheckedOutEnabled = SourceControlViewportUtils::GetFeedbackEnabled(ViewportWidget->GetViewportClient().Get(), ESourceControlStatus::CheckedOut);
-					bool bOpenForAddEnabled = SourceControlViewportUtils::GetFeedbackEnabled(ViewportWidget->GetViewportClient().Get(), ESourceControlStatus::OpenForAdd);
+					bool bNotAtHeadRevisionEnabled = false;
+					bool bCheckedOutByOtherUserEnabled = false;
+					bool bCheckedOutEnabled = false;
+					bool bOpenForAddEnabled = false;
+
+					if (TSharedPtr<SLevelViewport> ViewportWidgetPtr = ViewportWidget.Pin())
+					{
+						bNotAtHeadRevisionEnabled = SourceControlViewportUtils::GetFeedbackEnabled(ViewportWidgetPtr->GetViewportClient().Get(), ESourceControlStatus::NotAtHeadRevision);
+						bCheckedOutByOtherUserEnabled = SourceControlViewportUtils::GetFeedbackEnabled(ViewportWidgetPtr->GetViewportClient().Get(), ESourceControlStatus::CheckedOutByOtherUser);
+						bCheckedOutEnabled = SourceControlViewportUtils::GetFeedbackEnabled(ViewportWidgetPtr->GetViewportClient().Get(), ESourceControlStatus::CheckedOut);
+						bOpenForAddEnabled = SourceControlViewportUtils::GetFeedbackEnabled(ViewportWidgetPtr->GetViewportClient().Get(), ESourceControlStatus::OpenForAdd);
+					}
 
 					if (bNotAtHeadRevision && bNotAtHeadRevisionEnabled)
 					{
@@ -195,16 +201,16 @@ void FSourceControlViewportToolTips::InsertCanvas()
 			];
 
 			ViewportWidget = LevelViewport;
-			ViewportWidget->AddOverlayWidget(CanvasWidget.ToSharedRef());
+			ViewportWidget.Pin()->AddOverlayWidget(CanvasWidget.ToSharedRef());
 		}
 	}
 }
 
 void FSourceControlViewportToolTips::RemoveCanvas()
 {
-	if (ViewportWidget.IsValid())
+	if (TSharedPtr<SLevelViewport> ViewportWidgetPtr = ViewportWidget.Pin())
 	{
-		ViewportWidget->RemoveOverlayWidget(CanvasWidget.ToSharedRef());
+		ViewportWidgetPtr->RemoveOverlayWidget(CanvasWidget.ToSharedRef());
 	}
 
 	CanvasWidget.Reset();
@@ -220,15 +226,15 @@ FVector2D FSourceControlViewportToolTips::GetToolTipPosition() const
 	static const FVector2f TooltipOffsetFromMouse(12.0f, 8.0f);
 	static const FVector2f TooltipOffsetFromForceField(4.0f, 3.0f);
 
-	if (ViewportWidget.IsValid())
+	if (TSharedPtr<SLevelViewport> ViewportWidgetPtr = ViewportWidget.Pin())
 	{
 		FIntPoint ViewportOrigin;
 		FIntPoint ViewportSize;
-		ViewportWidget->GetViewportClient()->GetViewportDimensions(ViewportOrigin, ViewportSize);
+		ViewportWidgetPtr->GetViewportClient()->GetViewportDimensions(ViewportOrigin, ViewportSize);
 
 		int32 MouseX = 0;
 		int32 MouseY = 0;
-		if (FViewport* Viewport = ViewportWidget->GetViewportClient()->Viewport)
+		if (FViewport* Viewport = ViewportWidgetPtr->GetViewportClient()->Viewport)
 		{
 			MouseX = Viewport->GetMouseX();
 			MouseY = Viewport->GetMouseY();
diff --git a/Engine/Source/Developer/SourceControlViewport/Private/SourceControlViewportToolTips.h b/Engine/Source/Developer/SourceControlViewport/Private/SourceControlViewportToolTips.h
--- a/Engine/Source/Developer/SourceControlViewport/Private/SourceControlViewportToolTips.h
+++ b/Engine/Source/Developer/SourceControlViewport/Private/SourceControlViewportToolTips.h
@@ -20,6 +20,8 @@ public:
 
 public:
 	void Init();
+	void SetEnabled(bool bInEnabled);
+
 	bool Tick(float DeltaTime);
 
 private:
@@ -43,7 +45,7 @@ private:
 	TSharedPtr<SToolTip> ToolTipWidget;
 
 	/* The viewport in which the canvas widget is added */
-	TSharedPtr<SLevelViewport> ViewportWidget;
+	TWeakPtr<SLevelViewport> ViewportWidget;
 
 	/* The actor in the viewport that's currently hovered */
 	TWeakObjectPtr<AActor> Actor;
@@ -59,4 +61,7 @@ private:
 	double ActorTime = 0;
 	/* The time since the actor became relevant */
 	float DelayTime = 0;
+
+	/* The enabled state of the tooltips */
+	bool bEnabled = false;
 };
\ No newline at end of file
diff --git a/Engine/Source/Developer/SourceControlViewport/Private/SourceControlViewportUtils.cpp b/Engine/Source/Developer/SourceControlViewport/Private/SourceControlViewportUtils.cpp
--- a/Engine/Source/Developer/SourceControlViewport/Private/SourceControlViewportUtils.cpp
+++ b/Engine/Source/Developer/SourceControlViewport/Private/SourceControlViewportUtils.cpp
@@ -7,113 +7,98 @@
 namespace SourceControlViewportUtils
 {
 
-// Helper method to get the outline color index for a setting.
-static uint32 GetOutlineColorIndex(ESourceControlStatus Status)
+bool GetOverlaySetting(FViewportClient* ViewportClient, ESourceControlStatus Status)
 {
-	FLinearColor ColorToUse;
+	FString CVarName;
 	switch (Status)
 	{
 	case ESourceControlStatus::CheckedOutByOtherUser:
-		ColorToUse = FColor::Red;
+		CVarName = TEXT("SourceControl.Overlays.CheckedOutByOtherUser.Enable");
 		break;
 	case ESourceControlStatus::NotAtHeadRevision:
-		ColorToUse = FColor::Yellow;
+		CVarName = TEXT("SourceControl.Overlays.NotAtHeadRevision.Enable");
 		break;
 	case ESourceControlStatus::CheckedOut:
-		ColorToUse = FColor::Blue;
+		CVarName = TEXT("SourceControl.Overlays.CheckedOut.Enable");
 		break;
 	case ESourceControlStatus::OpenForAdd:
-		ColorToUse = FColor::Green;
+		CVarName = TEXT("SourceControl.Overlays.OpenForAdd.Enable");
 		break;
 	default:
 		checkNoEntry();
 		break;
 	}
 
-	// The available colors are defined in the UEditorStyleSettings::AdditionalSelectionColors array.
-	// The loop below gets the selection color index closest to the desired color.
-
-	int32 MinIndex = -1;
-	float MinDist = UE_MAX_FLT;
-
-	const UEditorStyleSettings* Settings = GetDefault<UEditorStyleSettings>();
-	for (int32 Index = 0; Index < UE_ARRAY_COUNT(Settings->AdditionalSelectionColors); ++Index)
+	if (IConsoleVariable* CVar = IConsoleManager::Get().FindConsoleVariable(*CVarName); ensure(CVar))
 	{
-		FLinearColor Color = Settings->AdditionalSelectionColors[Index];
-
-		float Dist = FLinearColor::Dist(ColorToUse, Color);
-		if (Dist < MinDist)
-		{
-			MinIndex = Index;
-			MinDist = Dist;
-		}
-	}
-
-	return MinIndex;
-}
-
-bool GetOutlineSetting(FViewportClient* ViewportClient, ESourceControlStatus Status)
-{
-	uint32 ColorIndex = GetOutlineColorIndex(Status);
-	switch (ColorIndex)
-	{
-	case 0:
-		return ViewportClient->GetEngineShowFlags()->SelectionOutlineColor0;
-	case 1:
-		return ViewportClient->GetEngineShowFlags()->SelectionOutlineColor1;
-	case 2:
-		return ViewportClient->GetEngineShowFlags()->SelectionOutlineColor2;
-	case 3:
-		return ViewportClient->GetEngineShowFlags()->SelectionOutlineColor3;
-	case 4:
-		return ViewportClient->GetEngineShowFlags()->SelectionOutlineColor4;
-	case 5:
-		return ViewportClient->GetEngineShowFlags()->SelectionOutlineColor5;
-	default:
-		ensure(false);
-		break;
+		return CVar->GetBool();
 	}
 
 	return false;
 }
 
-void SetOutlineSetting(FViewportClient* ViewportClient, ESourceControlStatus Status, bool bEnabled)
+void SetOverlaySetting(FViewportClient* ViewportClient, ESourceControlStatus Status, bool bEnabled)
 {
-	uint32 ColorIndex = GetOutlineColorIndex(Status);
-	switch (ColorIndex)
+	FString CVarName;
+	switch (Status)
 	{
-	case 0:
-		ViewportClient->GetEngineShowFlags()->SelectionOutlineColor0 = bEnabled;
-		break;
-	case 1:
-		ViewportClient->GetEngineShowFlags()->SelectionOutlineColor1 = bEnabled;
-		break;
-	case 2:
-		ViewportClient->GetEngineShowFlags()->SelectionOutlineColor2 = bEnabled;
+	case ESourceControlStatus::CheckedOutByOtherUser:
+		CVarName = TEXT("SourceControl.Overlays.CheckedOutByOtherUser.Enable");
 		break;
-	case 3:
-		ViewportClient->GetEngineShowFlags()->SelectionOutlineColor3 = bEnabled;
+	case ESourceControlStatus::NotAtHeadRevision:
+		CVarName = TEXT("SourceControl.Overlays.NotAtHeadRevision.Enable");
 		break;
-	case 4:
-		ViewportClient->GetEngineShowFlags()->SelectionOutlineColor4 = bEnabled;
+	case ESourceControlStatus::CheckedOut:
+		CVarName = TEXT("SourceControl.Overlays.CheckedOut.Enable");
 		break;
-	case 5:
-		ViewportClient->GetEngineShowFlags()->SelectionOutlineColor5 = bEnabled;
+	case ESourceControlStatus::OpenForAdd:
+		CVarName = TEXT("SourceControl.Overlays.OpenForAdd.Enable");
 		break;
 	default:
-		ensure(false);
+		checkNoEntry();
 		break;
 	}
+
+	if (IConsoleVariable* CVar = IConsoleManager::Get().FindConsoleVariable(*CVarName); ensure(CVar))
+	{
+		CVar->Set(bEnabled);
+	}
 }
 
 bool GetFeedbackEnabled(FViewportClient* ViewportClient, ESourceControlStatus Status)
 {
-	return GetOutlineSetting(ViewportClient, Status);
+	return GetOverlaySetting(ViewportClient, Status);
 }
 
 void SetFeedbackEnabled(FViewportClient* ViewportClient, ESourceControlStatus Status, bool bEnabled)
 {
-	SetOutlineSetting(ViewportClient, Status, bEnabled);
+	SetOverlaySetting(ViewportClient, Status, bEnabled);
+
+	GEditor->RedrawLevelEditingViewports();
+}
+
+uint8 GetFeedbackOpacity(FViewportClient* ViewportClient)
+{
+	uint8 Opacity = 128;
+
+	if (IConsoleVariable* CVar = IConsoleManager::Get().FindConsoleVariable(TEXT("SourceControl.Overlays.Alpha")); ensure(CVar))
+	{
+		Opacity = CVar->GetInt();
+	}
+
+	return FMath::Lerp<float>(0.f, 100.f, Opacity / 255.f);
+}
+
+void SetFeedbackOpacity(FViewportClient* ViewportClient, uint8 Opacity)
+{
+	Opacity = FMath::Lerp<float>(0.f, 255.f, Opacity / 100.f);
+
+	if (IConsoleVariable* CVar = IConsoleManager::Get().FindConsoleVariable(TEXT("SourceControl.Overlays.Alpha")); ensure(CVar))
+	{
+		CVar->Set(Opacity);
+	}
+
+	GEditor->RedrawLevelEditingViewports();
 }
 
 }
\ No newline at end of file
diff --git a/Engine/Source/Developer/SourceControlViewport/Private/SourceControlViewportUtils.h b/Engine/Source/Developer/SourceControlViewport/Private/SourceControlViewportUtils.h
--- a/Engine/Source/Developer/SourceControlViewport/Private/SourceControlViewportUtils.h
+++ b/Engine/Source/Developer/SourceControlViewport/Private/SourceControlViewportUtils.h
@@ -32,4 +32,10 @@ bool GetFeedbackEnabled(FViewportClient* ViewportClient, ESourceControlStatus St
 // Helper method to set the viewport feedback on or off on the viewport.
 void SetFeedbackEnabled(FViewportClient* ViewportClient, ESourceControlStatus Status, bool bEnabled);
 
+// Helper method to get the opacity value.
+uint8 GetFeedbackOpacity(FViewportClient* ViewportClient);
+
+// Helper method to set the opacity value.
+void SetFeedbackOpacity(FViewportClient* ViewportClient, uint8 Opacity);
+
 }
\ No newline at end of file
diff --git a/Engine/Source/Developer/SourceControlViewport/SourceControlViewport.Build.cs b/Engine/Source/Developer/SourceControlViewport/SourceControlViewport.Build.cs
--- a/Engine/Source/Developer/SourceControlViewport/SourceControlViewport.Build.cs
+++ b/Engine/Source/Developer/SourceControlViewport/SourceControlViewport.Build.cs
@@ -11,6 +11,7 @@ public class SourceControlViewport : ModuleRules
 				"Core",
 				"CoreUObject",
 				"Engine",
+				"InputCore",
 				"LevelEditor",
 				"Slate",
 				"SlateCore",
