Skein / ViewportStatus: Update tooltip show delay to the default Slate value.

diff --git a/Engine/Source/Developer/SourceControlViewport/Private/SourceControlViewportToolTips.cpp b/Engine/Source/Developer/SourceControlViewport/Private/SourceControlViewportToolTips.cpp
--- a/Engine/Source/Developer/SourceControlViewport/Private/SourceControlViewportToolTips.cpp
+++ b/Engine/Source/Developer/SourceControlViewport/Private/SourceControlViewportToolTips.cpp
@@ -108,7 +108,7 @@ void FSourceControlViewportToolTips::UpdateCanvas(float DeltaTime)
 			if (Actor.IsValid())
 			{
 				DelayTime += DeltaTime;
-				if (DelayTime >= 0.5f)
+				if (DelayTime >= 0.15f) // See: Slate.TooltipSummonDelay
 				{
 					UpdateToolTip();
 				}
