Skein / ViewportStatus: Possible fix for divide by zero crash.

When GetHitProxy is called on a viewport that doesn't have any hitmaps cached and doesn't have any slate buffers reserved yet, the code ends up doing a divide by zero in FSceneViewport::EnqueueBeginRenderFrame. This is a quick attempt to circumvent this.

diff --git a/Engine/Source/Developer/SourceControlViewport/Private/SourceControlViewportToolTips.cpp b/Engine/Source/Developer/SourceControlViewport/Private/SourceControlViewportToolTips.cpp
--- a/Engine/Source/Developer/SourceControlViewport/Private/SourceControlViewportToolTips.cpp
+++ b/Engine/Source/Developer/SourceControlViewport/Private/SourceControlViewportToolTips.cpp
@@ -76,43 +76,49 @@ void FSourceControlViewportToolTips::UpdateCanvas(float DeltaTime)
 		if (TSharedPtr<SLevelViewport> ViewportWidgetPtr = ViewportWidget.Pin())
 		{
 			FViewport* Viewport = ViewportWidgetPtr->GetViewportClient()->Viewport;
-			check(Viewport);
-
-			int32 MouseX = Viewport->GetMouseX();
-			int32 MouseY = Viewport->GetMouseY();
-			if (ActorMouseX != MouseX || ActorMouseY != MouseY)
+			if (Viewport && Viewport->GetRenderTargetTexture())
 			{
-				Actor.Reset();
-				UpdateToolTip();
+				int32 MouseX = Viewport->GetMouseX();
+				int32 MouseY = Viewport->GetMouseY();
+				if (ActorMouseX != MouseX || ActorMouseY != MouseY)
+				{
+					Actor.Reset();
+					UpdateToolTip();
 
-				ActorMouseX = MouseX;
-				ActorMouseY = MouseY;
-			}
+					ActorMouseX = MouseX;
+					ActorMouseY = MouseY;
+				}
 
-			HActor* ActorHitProxy = HitProxyCast<HActor>(Viewport->GetHitProxy(MouseX, MouseY));
-			if (ActorHitProxy == nullptr)
-			{
-				Actor.Reset();
-				UpdateToolTip();
-			}
-			if (ActorHitProxy != nullptr && ActorHitProxy->Actor != Actor)
-			{
-				Actor.Reset();
-				UpdateToolTip();
+				HActor* ActorHitProxy = HitProxyCast<HActor>(Viewport->GetHitProxy(MouseX, MouseY));
+				if (ActorHitProxy == nullptr)
+				{
+					Actor.Reset();
+					UpdateToolTip();
+				}
+				if (ActorHitProxy != nullptr && ActorHitProxy->Actor != Actor)
+				{
+					Actor.Reset();
+					UpdateToolTip();
 
-				Actor = MakeWeakObjectPtr(ActorHitProxy->Actor);
-				ActorTime = FPlatformTime::Seconds();
-				DelayTime = 0;
-			}
+					Actor = MakeWeakObjectPtr(ActorHitProxy->Actor);
+					ActorTime = FPlatformTime::Seconds();
+					DelayTime = 0;
+				}
 
-			if (Actor.IsValid())
-			{
-				DelayTime += DeltaTime;
-				if (DelayTime >= 0.15f) // See: Slate.TooltipSummonDelay
+				if (Actor.IsValid())
 				{
-					UpdateToolTip();
+					DelayTime += DeltaTime;
+					if (DelayTime >= 0.15f) // See: Slate.TooltipSummonDelay
+					{
+						UpdateToolTip();
+					}
 				}
 			}
+			else
+			{
+				Actor.Reset();
+				UpdateToolTip();
+			}
 		}
 	}
 }
