Skein / ViewportStatus: Only highlight actors stored in external package files (OFPA) to avoid the a checkout of the umap file resulting in the entire environment to be highlighted.

Additionally:
- Added an IsHidden check to the tooltip to avoid it showing up when the actor itself isn't visible.
- Moved the removal of the FTypedElementViewportOverlayColorColumn column outside the subquery so it's also removed in the case the sub query doesn't generate any results.

diff --git a/Engine/Plugins/Experimental/TypedElementsDataStorage/Source/TypedElementsDataStorageRevisionControl/Source/Private/RevisionControlProcessors.cpp b/Engine/Plugins/Experimental/TypedElementsDataStorage/Source/TypedElementsDataStorageRevisionControl/Source/Private/RevisionControlProcessors.cpp
--- a/Engine/Plugins/Experimental/TypedElementsDataStorage/Source/TypedElementsDataStorageRevisionControl/Source/Private/RevisionControlProcessors.cpp
+++ b/Engine/Plugins/Experimental/TypedElementsDataStorage/Source/TypedElementsDataStorageRevisionControl/Source/Private/RevisionControlProcessors.cpp
@@ -58,12 +58,14 @@ TAutoConsoleVariable<int32> CVarOverlayAlpha(
 	TEXT("Configures overlay opacity."),
 	ECVF_Default);
 
-static FColor DetermineOverlayColor(const TypedElementDataStorage::IQueryContext& ObjectContext, const TypedElementDataStorage::ICommonQueryContext& SCCContext)
+static FColor DetermineOverlayColor(const TypedElementDataStorage::IQueryContext& ObjectContext, const TypedElementDataStorage::ICommonQueryContext& SCCContext, const FMassActorFragment& Actor)
 {
 	check(IsInGameThread());
 
+	bool bExternal = Actor.Get() ? Actor.Get()->IsPackageExternal() : false;
+	bool bIgnored = !bExternal;
 	bool bSelected = ObjectContext.HasColumn<FTypedElementSelectionColumn>();
-	if (!bSelected)
+	if (!bIgnored && !bSelected)
 	{
 		// Check if the package is outdated because there is a newer version available.
 		if (SCCContext.HasColumn<FSCCNotCurrentTag>())
@@ -223,23 +225,20 @@ void UTypedElementRevisionControlFactory::RegisterApplyOverlays(ITypedElementDat
 			// This is in PrePhysics because the outline->actor query is in DuringPhysics and contexts don't flush changes between tick groups
 			FProcessor(DSI::EQueryTickPhase::PrePhysics, DataStorage.GetQueryTickGroupName(DSI::EQueryTickGroups::SyncExternalToDataStorage))
 				.ForceToGameThread(true),
-			[](DSI::IQueryContext& Context, TypedElementRowHandle ObjectRow, const FTypedElementPackageReference& PackageReference)
+			[](DSI::IQueryContext& Context, TypedElementRowHandle ObjectRow, const FMassActorFragment& Actor, const FTypedElementPackageReference& PackageReference)
 			{
+				Context.RemoveColumns<FTypedElementViewportOverlayColorColumn>(ObjectRow);
 				Context.RunSubquery(0, PackageReference.Row, CreateSubqueryCallbackBinding(
-					[&Context, &ObjectRow](DSI::ISubqueryContext& SubQueryContext)
+					[&Context, &ObjectRow, &Actor](DSI::ISubqueryContext& SubQueryContext)
 					{
-						FColor Color = DetermineOverlayColor(Context, SubQueryContext);
+						FColor Color = DetermineOverlayColor(Context, SubQueryContext, Actor);
 						if (Color.Bits != 0)
 						{
 							Context.AddColumn<FTypedElementViewportOverlayColorColumn>(ObjectRow, { .OverlayColor = Color });
 						}
-						else
-						{
-							Context.RemoveColumns<FTypedElementViewportOverlayColorColumn>(ObjectRow);
-						}
-						Context.AddColumns<FTypedElementSyncBackToWorldTag>(ObjectRow);
 					})
 				);
+				Context.AddColumns<FTypedElementSyncBackToWorldTag>(ObjectRow);
 			}
 		)
 		.DependsOn()
diff --git a/Engine/Source/Developer/SourceControlViewport/Private/SourceControlViewportToolTips.cpp b/Engine/Source/Developer/SourceControlViewport/Private/SourceControlViewportToolTips.cpp
--- a/Engine/Source/Developer/SourceControlViewport/Private/SourceControlViewportToolTips.cpp
+++ b/Engine/Source/Developer/SourceControlViewport/Private/SourceControlViewportToolTips.cpp
@@ -127,54 +127,61 @@ void FSourceControlViewportToolTips::UpdateToolTip()
 {
 	ToolTipText = FText::GetEmpty();
 
-	if (Actor.IsValid() && !Actor->IsSelected())
+	if (Actor.IsValid())
 	{
-		if (UPackage* Package = Actor->GetPackage())
+		bool bExternal = Actor->IsPackageExternal();
+		bool bIgnored = !bExternal;
+		bool bSelected = Actor->IsSelected();
+		bool bHidden = Actor->IsHidden();
+		if (!bIgnored && !bSelected && !bHidden)
 		{
-			FString SourceFileName = SourceControlHelpers::PackageFilename(Package);
-
-			ISourceControlModule& SourceControlModule = ISourceControlModule::Get();
-			ISourceControlProvider& SourceControlProvider = SourceControlModule.GetProvider();
-			if (SourceControlModule.IsEnabled())
+			if (UPackage* Package = Actor->GetPackage())
 			{
-				FSourceControlStatePtr State = SourceControlProvider.GetState(SourceFileName, EStateCacheUsage::Use);
-				if (State.IsValid() && State->IsSourceControlled())
-				{
-					FString Who;
-
-					bool bNotAtHeadRevision = !State->IsCurrent();
-					bool bCheckedOutByOtherUser = State->IsCheckedOutOther(&Who);
-					bool bCheckedOut = State->IsCheckedOut();
-					bool bOpenForAdd = State->IsAdded();
-
-					bool bNotAtHeadRevisionEnabled = false;
-					bool bCheckedOutByOtherUserEnabled = false;
-					bool bCheckedOutEnabled = false;
-					bool bOpenForAddEnabled = false;
-
-					if (TSharedPtr<SLevelViewport> ViewportWidgetPtr = ViewportWidget.Pin())
-					{
-						bNotAtHeadRevisionEnabled = SourceControlViewportUtils::GetFeedbackEnabled(ViewportWidgetPtr->GetViewportClient().Get(), ESourceControlStatus::NotAtHeadRevision);
-						bCheckedOutByOtherUserEnabled = SourceControlViewportUtils::GetFeedbackEnabled(ViewportWidgetPtr->GetViewportClient().Get(), ESourceControlStatus::CheckedOutByOtherUser);
-						bCheckedOutEnabled = SourceControlViewportUtils::GetFeedbackEnabled(ViewportWidgetPtr->GetViewportClient().Get(), ESourceControlStatus::CheckedOut);
-						bOpenForAddEnabled = SourceControlViewportUtils::GetFeedbackEnabled(ViewportWidgetPtr->GetViewportClient().Get(), ESourceControlStatus::OpenForAdd);
-					}
+				FString SourceFileName = SourceControlHelpers::PackageFilename(Package);
 
-					if (bNotAtHeadRevision && bNotAtHeadRevisionEnabled)
-					{
-						ToolTipText = LOCTEXT("NotAtHeadRevision", "File(s) out of sync");
-					}
-					else if (bCheckedOutByOtherUser && bCheckedOutByOtherUserEnabled)
-					{
-						ToolTipText = FText::Format(LOCTEXT("CheckedOutOtherUser", "File(s) checked out by {0}"), FText::FromString(Who));
-					}
-					else if (bCheckedOut && bCheckedOutEnabled)
-					{
-						ToolTipText = LOCTEXT("CheckedOut", "File(s) checked out by you");
-					}
-					else if (bOpenForAdd && bOpenForAddEnabled)
+				ISourceControlModule& SourceControlModule = ISourceControlModule::Get();
+				ISourceControlProvider& SourceControlProvider = SourceControlModule.GetProvider();
+				if (SourceControlModule.IsEnabled())
+				{
+					FSourceControlStatePtr State = SourceControlProvider.GetState(SourceFileName, EStateCacheUsage::Use);
+					if (State.IsValid() && State->IsSourceControlled())
 					{
-						ToolTipText = LOCTEXT("OpenForAdd", "File(s) added by you");
+						FString Who;
+
+						bool bNotAtHeadRevision = !State->IsCurrent();
+						bool bCheckedOutByOtherUser = State->IsCheckedOutOther(&Who);
+						bool bCheckedOut = State->IsCheckedOut();
+						bool bOpenForAdd = State->IsAdded();
+
+						bool bNotAtHeadRevisionEnabled = false;
+						bool bCheckedOutByOtherUserEnabled = false;
+						bool bCheckedOutEnabled = false;
+						bool bOpenForAddEnabled = false;
+
+						if (TSharedPtr<SLevelViewport> ViewportWidgetPtr = ViewportWidget.Pin())
+						{
+							bNotAtHeadRevisionEnabled = SourceControlViewportUtils::GetFeedbackEnabled(ViewportWidgetPtr->GetViewportClient().Get(), ESourceControlStatus::NotAtHeadRevision);
+							bCheckedOutByOtherUserEnabled = SourceControlViewportUtils::GetFeedbackEnabled(ViewportWidgetPtr->GetViewportClient().Get(), ESourceControlStatus::CheckedOutByOtherUser);
+							bCheckedOutEnabled = SourceControlViewportUtils::GetFeedbackEnabled(ViewportWidgetPtr->GetViewportClient().Get(), ESourceControlStatus::CheckedOut);
+							bOpenForAddEnabled = SourceControlViewportUtils::GetFeedbackEnabled(ViewportWidgetPtr->GetViewportClient().Get(), ESourceControlStatus::OpenForAdd);
+						}
+
+						if (bNotAtHeadRevision && bNotAtHeadRevisionEnabled)
+						{
+							ToolTipText = LOCTEXT("NotAtHeadRevision", "File(s) out of sync");
+						}
+						else if (bCheckedOutByOtherUser && bCheckedOutByOtherUserEnabled)
+						{
+							ToolTipText = FText::Format(LOCTEXT("CheckedOutOtherUser", "File(s) checked out by {0}"), FText::FromString(Who));
+						}
+						else if (bCheckedOut && bCheckedOutEnabled)
+						{
+							ToolTipText = LOCTEXT("CheckedOut", "File(s) checked out by you");
+						}
+						else if (bOpenForAdd && bOpenForAddEnabled)
+						{
+							ToolTipText = LOCTEXT("OpenForAdd", "File(s) added by you");
+						}
 					}
 				}
 			}
