Fix CIS ensure not finding renamed cvar SourceControl.Overlays.Enable
#rnx

diff --git a/Engine/Source/Developer/SourceControlViewport/Private/SourceControlViewportModule.cpp b/Engine/Source/Developer/SourceControlViewport/Private/SourceControlViewportModule.cpp
--- a/Engine/Source/Developer/SourceControlViewport/Private/SourceControlViewportModule.cpp
+++ b/Engine/Source/Developer/SourceControlViewport/Private/SourceControlViewportModule.cpp
@@ -70,7 +70,7 @@ void FSourceControlViewportModule::UpdateSettings()
 	ViewportMenu->SetEnabled(bEnabled);
 	ViewportToolTips->SetEnabled(bEnabled);
 
-	if (IConsoleVariable* CVar = IConsoleManager::Get().FindConsoleVariable(TEXT("SourceControl.Overlays.Enable")); ensure(CVar))
+	if (IConsoleVariable* CVar = IConsoleManager::Get().FindConsoleVariable(TEXT("RevisionControl.Overlays.Enable")); ensure(CVar))
 	{
 		CVar->Set(bEnabled);
 	}
