[Teds] Fix up some missed cvar lookup strings that were renamed to revision control to match desired nomenclature.

#rnx

diff --git a/Engine/Source/Developer/SourceControlViewport/Private/SourceControlViewportUtils.cpp b/Engine/Source/Developer/SourceControlViewport/Private/SourceControlViewportUtils.cpp
--- a/Engine/Source/Developer/SourceControlViewport/Private/SourceControlViewportUtils.cpp
+++ b/Engine/Source/Developer/SourceControlViewport/Private/SourceControlViewportUtils.cpp
@@ -13,16 +13,16 @@ bool GetOverlaySetting(FViewportClient* ViewportClient, ESourceControlStatus Sta
 	switch (Status)
 	{
 	case ESourceControlStatus::CheckedOutByOtherUser:
-		CVarName = TEXT("SourceControl.Overlays.CheckedOutByOtherUser.Enable");
+		CVarName = TEXT("RevisionControl.Overlays.CheckedOutByOtherUser.Enable");
 		break;
 	case ESourceControlStatus::NotAtHeadRevision:
-		CVarName = TEXT("SourceControl.Overlays.NotAtHeadRevision.Enable");
+		CVarName = TEXT("RevisionControl.Overlays.NotAtHeadRevision.Enable");
 		break;
 	case ESourceControlStatus::CheckedOut:
-		CVarName = TEXT("SourceControl.Overlays.CheckedOut.Enable");
+		CVarName = TEXT("RevisionControl.Overlays.CheckedOut.Enable");
 		break;
 	case ESourceControlStatus::OpenForAdd:
-		CVarName = TEXT("SourceControl.Overlays.OpenForAdd.Enable");
+		CVarName = TEXT("RevisionControl.Overlays.OpenForAdd.Enable");
 		break;
 	default:
 		checkNoEntry();
@@ -43,16 +43,16 @@ void SetOverlaySetting(FViewportClient* ViewportClient, ESourceControlStatus Sta
 	switch (Status)
 	{
 	case ESourceControlStatus::CheckedOutByOtherUser:
-		CVarName = TEXT("SourceControl.Overlays.CheckedOutByOtherUser.Enable");
+		CVarName = TEXT("RevisionControl.Overlays.CheckedOutByOtherUser.Enable");
 		break;
 	case ESourceControlStatus::NotAtHeadRevision:
-		CVarName = TEXT("SourceControl.Overlays.NotAtHeadRevision.Enable");
+		CVarName = TEXT("RevisionControl.Overlays.NotAtHeadRevision.Enable");
 		break;
 	case ESourceControlStatus::CheckedOut:
-		CVarName = TEXT("SourceControl.Overlays.CheckedOut.Enable");
+		CVarName = TEXT("RevisionControl.Overlays.CheckedOut.Enable");
 		break;
 	case ESourceControlStatus::OpenForAdd:
-		CVarName = TEXT("SourceControl.Overlays.OpenForAdd.Enable");
+		CVarName = TEXT("RevisionControl.Overlays.OpenForAdd.Enable");
 		break;
 	default:
 		checkNoEntry();
@@ -81,7 +81,7 @@ uint8 GetFeedbackOpacity(FViewportClient* ViewportClient)
 {
 	uint8 Opacity = 0;
 
-	if (IConsoleVariable* CVar = IConsoleManager::Get().FindConsoleVariable(TEXT("SourceControl.Overlays.Alpha")); ensure(CVar))
+	if (IConsoleVariable* CVar = IConsoleManager::Get().FindConsoleVariable(TEXT("RevisionControl.Overlays.Alpha")); ensure(CVar))
 	{
 		Opacity = CVar->GetInt();
 	}
@@ -91,7 +91,7 @@ uint8 GetFeedbackOpacity(FViewportClient* ViewportClient)
 
 void SetFeedbackOpacity(FViewportClient* ViewportClient, uint8 Opacity)
 {
-	if (IConsoleVariable* CVar = IConsoleManager::Get().FindConsoleVariable(TEXT("SourceControl.Overlays.Alpha")); ensure(CVar))
+	if (IConsoleVariable* CVar = IConsoleManager::Get().FindConsoleVariable(TEXT("RevisionControl.Overlays.Alpha")); ensure(CVar))
 	{
 		CVar->Set(Opacity);
 	}
