Deprecate the old level viewport toolbar.
- Deprecate the Level Viewport Toolbar context object.
- Expose & improve the Level Viewport Context object. Simplify common access patterns.
- Migrate & update client code.

[RN] Major Editor

diff --git a/Engine/Plugins/Editor/ModelingToolsEditorMode/Source/ModelingToolsEditorMode/Private/ModelingModeToolkit_Toolbars.cpp b/Engine/Plugins/Editor/ModelingToolsEditorMode/Source/ModelingToolsEditorMode/Private/ModelingModeToolkit_Toolbars.cpp
--- a/Engine/Plugins/Editor/ModelingToolsEditorMode/Source/ModelingToolsEditorMode/Private/ModelingModeToolkit_Toolbars.cpp
+++ b/Engine/Plugins/Editor/ModelingToolsEditorMode/Source/ModelingToolsEditorMode/Private/ModelingModeToolkit_Toolbars.cpp
@@ -444,26 +444,6 @@ TSharedRef<SWidget> MakeMenu_SelectionConfigSettings(const FModelingToolsEditorM
 }
 
 
-TSharedRef<SWidget> MakeMenu_SelectionEdits(FModelingToolsEditorModeToolkit* Toolkit)
-{
-	FMenuBuilder MenuBuilder(true, Toolkit->GetToolkitCommands());
-
-	MenuBuilder.BeginSection("Section_SelectionEdits", LOCTEXT("Section_SelectionEdits", "Selection Edits"));
-
-	const FModelingToolsManagerCommands& Commands = FModelingToolsManagerCommands::Get();
-
-	MenuBuilder.AddMenuEntry(Commands.BeginSelectionAction_SelectAll);
-	MenuBuilder.AddMenuEntry(Commands.BeginSelectionAction_ExpandToConnected);
-	MenuBuilder.AddMenuEntry(Commands.BeginSelectionAction_Invert);
-	MenuBuilder.AddMenuEntry(Commands.BeginSelectionAction_InvertConnected);
-	MenuBuilder.AddMenuEntry(Commands.BeginSelectionAction_Expand);
-	MenuBuilder.AddMenuEntry(Commands.BeginSelectionAction_Contract);
-
-	TSharedRef<SWidget> MenuWidget = MenuBuilder.MakeWidget();
-	return MenuWidget;
-}
-
-
 } // end namespace UELocal
 
 
@@ -621,159 +601,4 @@ void FModelingToolsEditorModeToolkit::ExtendSecondaryModeToolbar(UToolMenu *InMo
 #endif
 }
 
-
-
-
-void FModelingToolsEditorModeToolkit::MakeSelectionPaletteOverlayWidget()
-{
-	FVerticalToolBarBuilder ToolbarBuilder(
-		GetToolkitCommands(),
-		FMultiBoxCustomization::None,
-		TSharedPtr<FExtender>(), /*InForceSmallIcons=*/ true);
-	ToolbarBuilder.SetLabelVisibility(EVisibility::Collapsed);
-	ToolbarBuilder.SetStyle(FModelingToolsEditorModeStyle::Get().Get(), "SelectionToolBar");
-
-	const FModelingToolsManagerCommands& Commands = FModelingToolsManagerCommands::Get();
-
-	ToolbarBuilder.AddToolBarButton(Commands.MeshSelectionModeAction_NoSelection);
-	ToolbarBuilder.AddWidget(SNew(SSpacer).Size(FVector2D(1, 1)));
-	ToolbarBuilder.AddToolBarButton(Commands.MeshSelectionModeAction_MeshTriangles);
-	ToolbarBuilder.AddToolBarButton(Commands.MeshSelectionModeAction_MeshEdges);
-	ToolbarBuilder.AddToolBarButton(Commands.MeshSelectionModeAction_MeshVertices);
-	ToolbarBuilder.AddWidget(SNew(SSpacer).Size(FVector2D(1, 1)));
-	ToolbarBuilder.AddToolBarButton(Commands.MeshSelectionModeAction_GroupFaces);
-	ToolbarBuilder.AddToolBarButton(Commands.MeshSelectionModeAction_GroupEdges);
-	ToolbarBuilder.AddToolBarButton(Commands.MeshSelectionModeAction_GroupCorners);
-	ToolbarBuilder.AddWidget(SNew(SSpacer).Size(FVector2D(1, 8)));
-
-
-	ToolbarBuilder.AddWidget(
-		SNew(SComboButton)
-		.HasDownArrow(false)
-		.MenuPlacement(EMenuPlacement::MenuPlacement_MenuRight)
-		.ComboButtonStyle(FAppStyle::Get(), "SimpleComboButton")
-		.IsEnabled_Lambda([this]() { return IsInActiveTool() == false; })
-		.OnGetMenuContent(FOnGetContent::CreateLambda([this]()
-			{
-				return UELocal::MakeMenu_SelectionEdits(this);
-			}))
-		.ContentPadding(FMargin(1.0f, 1.0f))
-				.ButtonContent()
-				[
-					SNew(SImage)
-					.Image(FModelingToolsEditorModeStyle::Get()->GetBrush("ModelingModeSelection.Edits_Right"))
-				]);
-
-	ToolbarBuilder.AddWidget(
-		SNew(SComboButton)
-		.HasDownArrow(false)
-		.MenuPlacement(EMenuPlacement::MenuPlacement_MenuRight)
-		.ComboButtonStyle(FAppStyle::Get(), "SimpleComboButton")
-		.IsEnabled_Lambda([this]() { return IsInActiveTool() == false; })
-		.OnGetMenuContent(FOnGetContent::CreateLambda([this]()
-			{
-				return UELocal::MakeMenu_SelectionConfigSettings(this);
-			}))
-		.ContentPadding(FMargin(1.0f, 1.0f))
-				.ButtonContent()
-				[
-					SNew(SImage)
-					.Image(FModelingToolsEditorModeStyle::Get()->GetBrush("ModelingMode.DefaultSettings"))
-					.ColorAndOpacity(FSlateColor::UseForeground())
-				]);
-
-	ToolbarBuilder.AddWidget(SNew(SSpacer).Size(FVector2D(1, 16)));
-
-	//
-	// locking toggle button is implemented using 3 separate buttons because we cannot 
-	// control enabled/disabled or background color for a toolbar button
-	//
-
-	FUIAction UnlockTargetAction(
-		FExecuteAction::CreateLambda([this] { GetMeshSelectionManager()->SetCurrentTargetsLockState(false); }),
-		FCanExecuteAction::CreateLambda([this] { return GetMeshSelectionManager()->GetAnyCurrentTargetsLocked() && IsInActiveTool() == false; }),
-		FIsActionChecked::CreateLambda([] { return false; }),
-		FIsActionButtonVisible::CreateLambda([this] {
-			return GetMeshSelectionManager()->GetMeshTopologyMode() != UGeometrySelectionManager::EMeshTopologyMode::None
-			&& GetMeshSelectionManager()->GetAnyCurrentTargetsLockable()
-			&& GetMeshSelectionManager()->GetAnyCurrentTargetsLocked();
-			})
-	);
-	FUIAction LockTargetAction(
-		FExecuteAction::CreateLambda([this] { GetMeshSelectionManager()->SetCurrentTargetsLockState(true); }),
-		FCanExecuteAction::CreateLambda([this] { return !GetMeshSelectionManager()->GetAnyCurrentTargetsLocked() && IsInActiveTool() == false; }),
-		FIsActionChecked::CreateLambda([]() { return false; }),
-		FIsActionButtonVisible::CreateLambda([this] {
-			return GetMeshSelectionManager()->GetMeshTopologyMode() != UGeometrySelectionManager::EMeshTopologyMode::None
-			&& GetMeshSelectionManager()->GetAnyCurrentTargetsLockable()
-			&& (GetMeshSelectionManager()->GetAnyCurrentTargetsLocked() == false);
-			})
-	);
-	FUIAction DisabledLockTargetAction(
-		FExecuteAction::CreateLambda([] {}),
-		FCanExecuteAction::CreateLambda([] { return false; }),
-		FIsActionChecked::CreateLambda([]() { return false; }),
-		FIsActionButtonVisible::CreateLambda([this] {
-			return GetMeshSelectionManager()->GetMeshTopologyMode() == UGeometrySelectionManager::EMeshTopologyMode::None
-			|| GetMeshSelectionManager()->GetAnyCurrentTargetsLockable() == false;
-			})
-	);
-
-	ToolbarBuilder.BeginStyleOverride(FName("SelectionToolBar.RedButton"));
-	ToolbarBuilder.AddToolBarButton(UnlockTargetAction, NAME_None,
-		LOCTEXT("Selection_UnlockTarget", "Unlock"),
-		LOCTEXT("Selection_UnlockTarget_Tooltip", "Click to Unlock the Selected Object and allow Mesh Selections"),
-		FSlateIcon(FModelingToolsEditorModeStyle::GetStyleSetName(), "SelectionToolBarIcons.LockedTarget"));
-	ToolbarBuilder.EndStyleOverride();
-	ToolbarBuilder.BeginStyleOverride(FName("SelectionToolBar.GreenButton"));
-	ToolbarBuilder.AddToolBarButton(LockTargetAction, NAME_None,
-		LOCTEXT("Selection_LockTarget", "Lock"),
-		LOCTEXT("Selection_LockTarget_Tooltip", "Click to Lock the Selected Object"),
-		FSlateIcon(FModelingToolsEditorModeStyle::GetStyleSetName(), "SelectionToolBarIcons.UnlockedTarget"));
-	ToolbarBuilder.EndStyleOverride();
-	ToolbarBuilder.AddToolBarButton(DisabledLockTargetAction, NAME_None,
-		LOCTEXT("Selection_LockTargetDisabled", "(No Locking)"),
-		LOCTEXT("Selection_LockTargetDisabled_Tooltip", "No Active Selection Targets are Lockable"),
-		FSlateIcon(FModelingToolsEditorModeStyle::GetStyleSetName(), "SelectionToolBarIcons.UnlockedTarget"));
-
-
-	TSharedRef<SWidget> Toolbar = ToolbarBuilder.MakeWidget();
-
-	SAssignNew(SelectionPaletteOverlayWidget, SVerticalBox)
-		+ SVerticalBox::Slot()
-		.HAlign(HAlign_Left)
-		.VAlign(VAlign_Center)
-		.Padding(FMargin(2.0f, 0.0f, 0.f, 0.f))
-		[
-			SNew(SBorder)
-			.BorderImage(FAppStyle::Get().GetBrush("EditorViewport.OverlayBrush"))
-			.Padding(FMargin(3.0f, 6.0f, 3.f, 6.f))
-			[
-				SNew(SVerticalBox)
-				+ SVerticalBox::Slot()
-				.AutoHeight()
-				.Padding(FMargin(0.0, 0.f, 0.f, 0.f))
-				[
-					Toolbar
-				]
-			]
-		];
-
-	SelectionPaletteOverlayWidget->SetVisibility( TAttribute<EVisibility>::CreateLambda([this]()
-	{
-		if ( UModelingToolsEditorMode* ModelingMode = Cast<UModelingToolsEditorMode>(GetScriptableEditorMode()) )
-		{
-			// only show old MES Toolbar (on left side) when exclusively showing the Old Toolbar (& MES is enabled)
-			// TODO: when new viewport toolbar is finalized remove the ShowNewViewportToolbars condition
-			return ModelingMode->GetMeshElementSelectionSystemEnabled() && !UE::UnrealEd::ShowNewViewportToolbars() ? EVisibility::SelfHitTestInvisible : EVisibility::Collapsed;
-		}
-		return EVisibility::Collapsed;
-	}) );
-
-
-}
-
-
-
-
 #undef LOCTEXT_NAMESPACE
diff --git a/Engine/Plugins/Editor/ModelingToolsEditorMode/Source/ModelingToolsEditorMode/Private/ModelingToolsEditorModeToolkit.cpp b/Engine/Plugins/Editor/ModelingToolsEditorMode/Source/ModelingToolsEditorMode/Private/ModelingToolsEditorModeToolkit.cpp
--- a/Engine/Plugins/Editor/ModelingToolsEditorMode/Source/ModelingToolsEditorMode/Private/ModelingToolsEditorModeToolkit.cpp
+++ b/Engine/Plugins/Editor/ModelingToolsEditorMode/Source/ModelingToolsEditorMode/Private/ModelingToolsEditorModeToolkit.cpp
@@ -191,12 +191,6 @@ FModelingToolsEditorModeToolkit::~FModelingToolsEditorModeToolkit()
 {
 	if (IsHosted())
 	{
-		if (SelectionPaletteOverlayWidget.IsValid())
-		{
-			GetToolkitHost()->RemoveViewportOverlayWidget(SelectionPaletteOverlayWidget.ToSharedRef());
-			SelectionPaletteOverlayWidget.Reset();
-		}
-
 		if (GizmoNumericalUIOverlayWidget.IsValid())
 		{
 			GetToolkitHost()->RemoveViewportOverlayWidget(GizmoNumericalUIOverlayWidget.ToSharedRef());
@@ -661,9 +655,6 @@ void FModelingToolsEditorModeToolkit::Init(const TSharedPtr<IToolkitHost>& InitT
 	MakeGizmoNumericalUIOverlayWidget();
 	GetToolkitHost()->AddViewportOverlayWidget(GizmoNumericalUIOverlayWidget.ToSharedRef());
 
-	MakeSelectionPaletteOverlayWidget();
-	GetToolkitHost()->AddViewportOverlayWidget(SelectionPaletteOverlayWidget.ToSharedRef());
-
 	CurrentPresetPath = FSoftObjectPath(); // Default to the default collection by leaving this null.
 
 	UModelingToolsModeCustomizationSettings* ModelingCustomizationSettings = GetMutableDefault<UModelingToolsModeCustomizationSettings>();
@@ -2424,11 +2415,6 @@ void FModelingToolsEditorModeToolkit::OnActiveViewportChanged(TSharedPtr<IAssetV
 			{
 				GetToolkitHost()->RemoveViewportOverlayWidget(GizmoNumericalUIOverlayWidget.ToSharedRef(), OldViewport);
 			}
-
-			if (SelectionPaletteOverlayWidget.IsValid())
-			{
-				GetToolkitHost()->RemoveViewportOverlayWidget(SelectionPaletteOverlayWidget.ToSharedRef(), OldViewport);
-			}
 		}
 
 		// Add the hud to the new viewport
@@ -2438,11 +2424,6 @@ void FModelingToolsEditorModeToolkit::OnActiveViewportChanged(TSharedPtr<IAssetV
 		{
 			GetToolkitHost()->AddViewportOverlayWidget(GizmoNumericalUIOverlayWidget.ToSharedRef(), NewViewport);
 		}
-
-		if (SelectionPaletteOverlayWidget.IsValid())
-		{
-			GetToolkitHost()->AddViewportOverlayWidget(SelectionPaletteOverlayWidget.ToSharedRef(), NewViewport);
-		}
 	}
 
 #if ENABLE_STYLUS_SUPPORT
diff --git a/Engine/Plugins/Editor/ModelingToolsEditorMode/Source/ModelingToolsEditorMode/Public/ModelingToolsEditorModeToolkit.h b/Engine/Plugins/Editor/ModelingToolsEditorMode/Source/ModelingToolsEditorMode/Public/ModelingToolsEditorModeToolkit.h
--- a/Engine/Plugins/Editor/ModelingToolsEditorMode/Source/ModelingToolsEditorMode/Public/ModelingToolsEditorModeToolkit.h
+++ b/Engine/Plugins/Editor/ModelingToolsEditorMode/Source/ModelingToolsEditorMode/Public/ModelingToolsEditorModeToolkit.h
@@ -149,9 +149,6 @@ private:
 	TSharedPtr<SWidget> ToolShutdownViewportOverlayWidget;
 	void MakeToolShutdownOverlayWidget();
 	
-	TSharedPtr<SWidget> SelectionPaletteOverlayWidget;
-	void MakeSelectionPaletteOverlayWidget();
-
 	void MakeGizmoNumericalUIOverlayWidget();
 	TSharedPtr<STransformGizmoNumericalUIOverlay> GizmoNumericalUIOverlayWidget;
 
diff --git a/Engine/Source/Developer/SourceControlViewport/Private/SourceControlViewportMenu.cpp b/Engine/Source/Developer/SourceControlViewport/Private/SourceControlViewportMenu.cpp
--- a/Engine/Source/Developer/SourceControlViewport/Private/SourceControlViewportMenu.cpp
+++ b/Engine/Source/Developer/SourceControlViewport/Private/SourceControlViewportMenu.cpp
@@ -6,8 +6,8 @@
 #include "EngineAnalytics.h"
 #include "LevelEditorMenuContext.h"
 #include "LevelEditorViewport.h"
+#include "SLevelViewport.h"
 #include "ToolMenus.h"
-#include "ViewportToolbar/UnrealEdViewportToolbar.h"
 
 #define LOCTEXT_NAMESPACE "SourceControlViewportMenu"
 
@@ -57,153 +57,120 @@ void FSourceControlViewportMenu::InsertViewportMenu()
 void FSourceControlViewportMenu::PopulateViewportMenu(UToolMenu* InMenu)
 {
 	check(InMenu);
+	
+	FToolMenuSection& RevisionControlSection = InMenu->FindOrAddSection(TEXT("AllShowFlags"));
+	
+	RevisionControlSection.AddSubMenu(
+		SubMenuName,
+		LOCTEXT("RevisionControlSubMenu", "Revision Control"),
+		LOCTEXT("RevisionControlSubMenu_ToolTip", "Toggle revision control viewport options on or off."),
+		FNewToolMenuDelegate::CreateSP(this, &FSourceControlViewportMenu::PopulateRevisionControlMenu),
+		false,
+		FSlateIcon(FAppStyle::Get().GetStyleSetName(), "ShowFlagsMenu.SubMenu.RevisionControl")
+	);
+}
 
-	InMenu->AddDynamicSection(SectionName, FNewToolMenuDelegate::CreateLambda(
-		[this](UToolMenu* InMenu)
-		{
-			ULevelViewportToolBarContext* Context = InMenu->FindContext<ULevelViewportToolBarContext>();
-			if (!Context)
-			{
-				return;
-			}
-
-			FLevelEditorViewportClient* ViewportClient = Context->GetLevelViewportClient();
-			if (!ViewportClient || !ViewportClient->IsPerspective())
-			{
-				return;
-			}
-
-			OpacityWidget = SNew(SSpinBox<uint8>)
-				.ClearKeyboardFocusOnCommit(true)
-				.OnValueChanged_Lambda(
-					[this, ViewportClient](uint8 InNewValue)
-					{
-						SetOpacityValue(ViewportClient, InNewValue);
-					}
-				)
-				.OnValueCommitted_Lambda(
-					[this, ViewportClient](uint8 InNewValue, ETextCommit::Type InCommitType)
-					{
-						SetOpacityValue(ViewportClient, InNewValue);
-					}
-				)
-				.Value_Lambda(
-					[this, ViewportClient]()
-					{
-						return GetOpacityValue(ViewportClient);
-					}
-				)
-				.MinValue(0)
-				.MinSliderValue(0)
-				.MaxValue(100)
-				.MaxSliderValue(100);
-
-			FToolMenuSection& RevisionControlSection = UE::UnrealEd::ShowNewViewportToolbars()
-				? InMenu->FindOrAddSection(TEXT("AllShowFlags"))
-				: InMenu->FindOrAddSection(TEXT("LevelViewportEditorShow"));
-			
-			RevisionControlSection.AddDynamicEntry(NAME_None, FNewToolMenuSectionDelegate::CreateLambda(
-				[this, ViewportClient](FToolMenuSection& InSection)
-				{
-					InSection.AddSubMenu(
-						SubMenuName,
-						LOCTEXT("RevisionControlSubMenu", "Revision Control"),
-						LOCTEXT("RevisionControlSubMenu_ToolTip", "Toggle revision control viewport options on or off."),
-						FNewToolMenuDelegate::CreateLambda(
-							[this, ViewportClient](UToolMenu* InSubMenu)
-							{
-								FToolMenuSection& DefaultSection = InSubMenu->AddSection(NAME_None, LOCTEXT("RevisionControlSectionStatus", "Status Highlights"));
-
-								DefaultSection.AddMenuEntry(
-									NAME_None,
-									LOCTEXT("ShowAll", "Show All"),
-									LOCTEXT("ShowAll_ToolTip", "Enable highlighting for all statuses"),
-									FSlateIcon(),
-									FUIAction(
-										FExecuteAction::CreateSP(this, &FSourceControlViewportMenu::ShowAll, ViewportClient)
-									),
-									EUserInterfaceActionType::Button
-								);
-
-								DefaultSection.AddMenuEntry(
-									NAME_None,
-									LOCTEXT("HideAll", "Hide All"),
-									LOCTEXT("HideAll_ToolTip", "Disable highlighting for all statuses"),
-									FSlateIcon(),
-									FUIAction(
-										FExecuteAction::CreateSP(this, &FSourceControlViewportMenu::HideAll, ViewportClient)
-									),
-									EUserInterfaceActionType::Button
-								);
-
-								DefaultSection.AddSeparator(NAME_None);
-
-								DefaultSection.AddMenuEntry(
-									NAME_None,
-									LOCTEXT("HighlightCheckedOutByOtherUser", "Checked Out by Others"),
-									LOCTEXT("HighlightCheckedOutByOtherUser_ToolTip", "Highlight objects that are checked out by someone else."),
-									FSlateIcon(FRevisionControlStyleManager::GetStyleSetName(), "RevisionControl.ShowMenu.CheckedOutByOtherUser"),
-									FUIAction(
-										FExecuteAction::CreateSP(this, &FSourceControlViewportMenu::ToggleHighlight, ViewportClient, ESourceControlStatus::CheckedOutByOtherUser),
-										FCanExecuteAction(),
-										FIsActionChecked::CreateSP(this, &FSourceControlViewportMenu::IsHighlighted, ViewportClient, ESourceControlStatus::CheckedOutByOtherUser)
-									),
-									EUserInterfaceActionType::ToggleButton
-								);
-
-								DefaultSection.AddMenuEntry(
-									NAME_None,
-									LOCTEXT("HighlightNotAtHeadRevision", "Out of Date"),
-									LOCTEXT("HighlightNotAtHeadRevision_ToolTip", "Highlight objects that are not at the latest revision."),
-									FSlateIcon(FRevisionControlStyleManager::GetStyleSetName(), "RevisionControl.ShowMenu.NotAtHeadRevision"),
-									FUIAction(
-										FExecuteAction::CreateSP(this, &FSourceControlViewportMenu::ToggleHighlight, ViewportClient, ESourceControlStatus::NotAtHeadRevision),
-										FCanExecuteAction(),
-										FIsActionChecked::CreateSP(this, &FSourceControlViewportMenu::IsHighlighted, ViewportClient, ESourceControlStatus::NotAtHeadRevision)
-									),
-									EUserInterfaceActionType::ToggleButton
-								);
-
-								DefaultSection.AddMenuEntry(
-									NAME_None,
-									LOCTEXT("HighlightCheckedOut", "Checked Out by Me"),
-									LOCTEXT("HighlightCheckedOut_ToolTip", "Highlight objects that are checked out by me."),
-									FSlateIcon(FRevisionControlStyleManager::GetStyleSetName(), "RevisionControl.ShowMenu.CheckedOut"),
-									FUIAction(
-										FExecuteAction::CreateSP(this, &FSourceControlViewportMenu::ToggleHighlight, ViewportClient, ESourceControlStatus::CheckedOut),
-										FCanExecuteAction(),
-										FIsActionChecked::CreateSP(this, &FSourceControlViewportMenu::IsHighlighted, ViewportClient, ESourceControlStatus::CheckedOut)
-									),
-									EUserInterfaceActionType::ToggleButton
-								);
-
-								DefaultSection.AddMenuEntry(
-									NAME_None,
-									LOCTEXT("HighlightOpenForAdd", "Newly Added"),
-									LOCTEXT("HighlightOpenForAdd_ToolTip", "Highlight objects that have been added by me."),
-									FSlateIcon(FRevisionControlStyleManager::GetStyleSetName(), "RevisionControl.ShowMenu.OpenForAdd"),
-									FUIAction(
-										FExecuteAction::CreateSP(this, &FSourceControlViewportMenu::ToggleHighlight, ViewportClient, ESourceControlStatus::OpenForAdd),
-										FCanExecuteAction(),
-										FIsActionChecked::CreateSP(this, &FSourceControlViewportMenu::IsHighlighted, ViewportClient, ESourceControlStatus::OpenForAdd)
-									),
-									EUserInterfaceActionType::ToggleButton
-								);
-
-								DefaultSection.AddEntry(FToolMenuEntry::InitWidget(
-									NAME_None,
-									OpacityWidget.ToSharedRef(),
-									LOCTEXT("Opacity", "Opacity")
-									)
-								);
-							}
-						),
-						false,
-						FSlateIcon(FAppStyle::Get().GetStyleSetName(), "ShowFlagsMenu.SubMenu.RevisionControl")
-					);
-				}));
-		}
-	)
+void FSourceControlViewportMenu::PopulateRevisionControlMenu(UToolMenu* InMenu)
+{
+	TSharedPtr<SLevelViewport> Viewport = ULevelViewportContext::GetLevelViewport(InMenu);
+	if (!Viewport)
+	{
+		return;
+	}
+
+	FToolMenuSection& DefaultSection = InMenu->AddSection(NAME_None, LOCTEXT("RevisionControlSectionStatus", "Status Highlights"));
+
+	DefaultSection.AddMenuEntry(
+		NAME_None,
+		LOCTEXT("ShowAll", "Show All"),
+		LOCTEXT("ShowAll_ToolTip", "Enable highlighting for all statuses"),
+		FSlateIcon(),
+		FUIAction(
+			FExecuteAction::CreateSP(this, &FSourceControlViewportMenu::ShowAll, Viewport.ToWeakPtr())
+		),
+		EUserInterfaceActionType::Button
+	);
+
+	DefaultSection.AddMenuEntry(
+		NAME_None,
+		LOCTEXT("HideAll", "Hide All"),
+		LOCTEXT("HideAll_ToolTip", "Disable highlighting for all statuses"),
+		FSlateIcon(),
+		FUIAction(
+			FExecuteAction::CreateSP(this, &FSourceControlViewportMenu::HideAll, Viewport.ToWeakPtr())
+		),
+		EUserInterfaceActionType::Button
+	);
+
+	DefaultSection.AddSeparator(NAME_None);
+
+	DefaultSection.AddMenuEntry(
+		NAME_None,
+		LOCTEXT("HighlightCheckedOutByOtherUser", "Checked Out by Others"),
+		LOCTEXT("HighlightCheckedOutByOtherUser_ToolTip", "Highlight objects that are checked out by someone else."),
+		FSlateIcon(FRevisionControlStyleManager::GetStyleSetName(), "RevisionControl.ShowMenu.CheckedOutByOtherUser"),
+		FUIAction(
+			FExecuteAction::CreateSP(this, &FSourceControlViewportMenu::ToggleHighlight, Viewport.ToWeakPtr(), ESourceControlStatus::CheckedOutByOtherUser),
+			FCanExecuteAction(),
+			FIsActionChecked::CreateSP(this, &FSourceControlViewportMenu::IsHighlighted, Viewport.ToWeakPtr(), ESourceControlStatus::CheckedOutByOtherUser)
+		),
+		EUserInterfaceActionType::ToggleButton
+	);
+
+	DefaultSection.AddMenuEntry(
+		NAME_None,
+		LOCTEXT("HighlightNotAtHeadRevision", "Out of Date"),
+		LOCTEXT("HighlightNotAtHeadRevision_ToolTip", "Highlight objects that are not at the latest revision."),
+		FSlateIcon(FRevisionControlStyleManager::GetStyleSetName(), "RevisionControl.ShowMenu.NotAtHeadRevision"),
+		FUIAction(
+			FExecuteAction::CreateSP(this, &FSourceControlViewportMenu::ToggleHighlight, Viewport.ToWeakPtr(), ESourceControlStatus::NotAtHeadRevision),
+			FCanExecuteAction(),
+			FIsActionChecked::CreateSP(this, &FSourceControlViewportMenu::IsHighlighted, Viewport.ToWeakPtr(), ESourceControlStatus::NotAtHeadRevision)
+		),
+		EUserInterfaceActionType::ToggleButton
+	);
+
+	DefaultSection.AddMenuEntry(
+		NAME_None,
+		LOCTEXT("HighlightCheckedOut", "Checked Out by Me"),
+		LOCTEXT("HighlightCheckedOut_ToolTip", "Highlight objects that are checked out by me."),
+		FSlateIcon(FRevisionControlStyleManager::GetStyleSetName(), "RevisionControl.ShowMenu.CheckedOut"),
+		FUIAction(
+			FExecuteAction::CreateSP(this, &FSourceControlViewportMenu::ToggleHighlight, Viewport.ToWeakPtr(), ESourceControlStatus::CheckedOut),
+			FCanExecuteAction(),
+			FIsActionChecked::CreateSP(this, &FSourceControlViewportMenu::IsHighlighted, Viewport.ToWeakPtr(), ESourceControlStatus::CheckedOut)
+		),
+		EUserInterfaceActionType::ToggleButton
+	);
+
+	DefaultSection.AddMenuEntry(
+		NAME_None,
+		LOCTEXT("HighlightOpenForAdd", "Newly Added"),
+		LOCTEXT("HighlightOpenForAdd_ToolTip", "Highlight objects that have been added by me."),
+		FSlateIcon(FRevisionControlStyleManager::GetStyleSetName(), "RevisionControl.ShowMenu.OpenForAdd"),
+		FUIAction(
+			FExecuteAction::CreateSP(this, &FSourceControlViewportMenu::ToggleHighlight, Viewport.ToWeakPtr(), ESourceControlStatus::OpenForAdd),
+			FCanExecuteAction(),
+			FIsActionChecked::CreateSP(this, &FSourceControlViewportMenu::IsHighlighted, Viewport.ToWeakPtr(), ESourceControlStatus::OpenForAdd)
+		),
+		EUserInterfaceActionType::ToggleButton
+	);
+
+	TSharedRef<SSpinBox<uint8>> OpacityWidget = SNew(SSpinBox<uint8>)
+		.ClearKeyboardFocusOnCommit(true)
+		.OnValueChanged(this, &FSourceControlViewportMenu::SetOpacityValue, Viewport.ToWeakPtr())
+		.OnValueCommitted(this, &FSourceControlViewportMenu::OnOpacityCommitted, Viewport.ToWeakPtr())
+		.Value(this, &FSourceControlViewportMenu::GetOpacityValue, Viewport.ToWeakPtr())
+		.MinValue(0)
+		.MinSliderValue(0)
+		.MaxValue(100)
+		.MaxSliderValue(100);
+
+	DefaultSection.AddEntry(FToolMenuEntry::InitWidget(
+		NAME_None,
+		OpacityWidget,
+		LOCTEXT("Opacity", "Opacity")
+		)
 	);
 }
 
@@ -219,58 +186,79 @@ void FSourceControlViewportMenu::RemoveViewportMenu()
 	}
 }
 
-void FSourceControlViewportMenu::ShowAll(FLevelEditorViewportClient* ViewportClient)
+void FSourceControlViewportMenu::ShowAll(TWeakPtr<SLevelViewport> Viewport)
 {
-	ensure(ViewportClient);
-
-	SourceControlViewportUtils::SetFeedbackEnabled(ViewportClient, ESourceControlStatus::CheckedOutByOtherUser, /*bEnabled=*/true);
-	SourceControlViewportUtils::SetFeedbackEnabled(ViewportClient, ESourceControlStatus::NotAtHeadRevision, /*bEnabled=*/true);
-	SourceControlViewportUtils::SetFeedbackEnabled(ViewportClient, ESourceControlStatus::CheckedOut, /*bEnabled=*/true);
-	SourceControlViewportUtils::SetFeedbackEnabled(ViewportClient, ESourceControlStatus::OpenForAdd, /*bEnabled=*/true);
+	if (TSharedPtr<SLevelViewport> Pinned = Viewport.Pin())
+	{
+		FLevelEditorViewportClient* ViewportClient = &Pinned->GetLevelViewportClient();
 	
-	RecordToggleEvent(TEXT("All"), /*bEnabled=*/true);
+		SourceControlViewportUtils::SetFeedbackEnabled(ViewportClient, ESourceControlStatus::CheckedOutByOtherUser, /*bEnabled=*/true);
+		SourceControlViewportUtils::SetFeedbackEnabled(ViewportClient, ESourceControlStatus::NotAtHeadRevision, /*bEnabled=*/true);
+		SourceControlViewportUtils::SetFeedbackEnabled(ViewportClient, ESourceControlStatus::CheckedOut, /*bEnabled=*/true);
+		SourceControlViewportUtils::SetFeedbackEnabled(ViewportClient, ESourceControlStatus::OpenForAdd, /*bEnabled=*/true);
+	
+		RecordToggleEvent(TEXT("All"), /*bEnabled=*/true);
+	}
 }
 
-void FSourceControlViewportMenu::HideAll(FLevelEditorViewportClient* ViewportClient)
+void FSourceControlViewportMenu::HideAll(TWeakPtr<SLevelViewport> Viewport)
 {
-	ensure(ViewportClient);
-
-	SourceControlViewportUtils::SetFeedbackEnabled(ViewportClient, ESourceControlStatus::CheckedOutByOtherUser, /*bEnabled=*/false);
-	SourceControlViewportUtils::SetFeedbackEnabled(ViewportClient, ESourceControlStatus::NotAtHeadRevision, /*bEnabled=*/false);
-	SourceControlViewportUtils::SetFeedbackEnabled(ViewportClient, ESourceControlStatus::CheckedOut, /*bEnabled=*/false);
-	SourceControlViewportUtils::SetFeedbackEnabled(ViewportClient, ESourceControlStatus::OpenForAdd, /*bEnabled=*/false);
+	if (TSharedPtr<SLevelViewport> Pinned = Viewport.Pin())
+	{
+		FLevelEditorViewportClient* ViewportClient = &Pinned->GetLevelViewportClient();
+		SourceControlViewportUtils::SetFeedbackEnabled(ViewportClient, ESourceControlStatus::CheckedOutByOtherUser, /*bEnabled=*/false);
+		SourceControlViewportUtils::SetFeedbackEnabled(ViewportClient, ESourceControlStatus::NotAtHeadRevision, /*bEnabled=*/false);
+		SourceControlViewportUtils::SetFeedbackEnabled(ViewportClient, ESourceControlStatus::CheckedOut, /*bEnabled=*/false);
+		SourceControlViewportUtils::SetFeedbackEnabled(ViewportClient, ESourceControlStatus::OpenForAdd, /*bEnabled=*/false);
 
-	RecordToggleEvent(TEXT("All"), /*bEnabled=*/false);
+		RecordToggleEvent(TEXT("All"), /*bEnabled=*/false);
+	}
 }
 
-void FSourceControlViewportMenu::ToggleHighlight(FLevelEditorViewportClient* ViewportClient, ESourceControlStatus Status)
+void FSourceControlViewportMenu::ToggleHighlight(TWeakPtr<SLevelViewport> Viewport, ESourceControlStatus Status)
 {
-	ensure(ViewportClient);
-
-	bool bOld = SourceControlViewportUtils::GetFeedbackEnabled(ViewportClient, Status);
-	bool bNew = bOld ? false : true;
-	SourceControlViewportUtils::SetFeedbackEnabled(ViewportClient, Status, bNew);
-
-	FString EnumValueWithoutType = UEnum::GetValueAsString(Status)
-		.Replace(TEXT("ESourceControlStatus::"), TEXT(""));
-	RecordToggleEvent(EnumValueWithoutType, bNew);
+	if (TSharedPtr<SLevelViewport> Pinned = Viewport.Pin())
+	{
+		FLevelEditorViewportClient* ViewportClient = &Pinned->GetLevelViewportClient();
+		bool bOld = SourceControlViewportUtils::GetFeedbackEnabled(ViewportClient, Status);
+		bool bNew = bOld ? false : true;
+		SourceControlViewportUtils::SetFeedbackEnabled(ViewportClient, Status, bNew);
+
+		FString EnumValueWithoutType = UEnum::GetValueAsString(Status)
+			.Replace(TEXT("ESourceControlStatus::"), TEXT(""));
+		RecordToggleEvent(EnumValueWithoutType, bNew);
+	}
 }
 
-bool FSourceControlViewportMenu::IsHighlighted(FLevelEditorViewportClient* ViewportClient, ESourceControlStatus Status) const
+bool FSourceControlViewportMenu::IsHighlighted(TWeakPtr<SLevelViewport> Viewport, ESourceControlStatus Status) const
 {
-	ensure(ViewportClient);
+	if (TSharedPtr<SLevelViewport> Pinned = Viewport.Pin())
+	{
+		return SourceControlViewportUtils::GetFeedbackEnabled(&Pinned->GetLevelViewportClient(), Status);	
+	}
+	return false;
+}
 
-	return SourceControlViewportUtils::GetFeedbackEnabled(ViewportClient, Status);
+void FSourceControlViewportMenu::OnOpacityCommitted(uint8 NewValue, ETextCommit::Type CommitType, TWeakPtr<SLevelViewport> Viewport)
+{
+	SetOpacityValue(NewValue, Viewport);
 }
 
-void FSourceControlViewportMenu::SetOpacityValue(FLevelEditorViewportClient* ViewportClient, uint8 InNewValue)
+void FSourceControlViewportMenu::SetOpacityValue(uint8 NewValue, TWeakPtr<SLevelViewport> Viewport)
 {
-	SourceControlViewportUtils::SetFeedbackOpacity(ViewportClient, InNewValue);
+	if (TSharedPtr<SLevelViewport> Pinned = Viewport.Pin())
+	{
+		SourceControlViewportUtils::SetFeedbackOpacity(&Pinned->GetLevelViewportClient(), NewValue);
+	}
 }
 
-uint8 FSourceControlViewportMenu::GetOpacityValue(FLevelEditorViewportClient* ViewportClient) const
+uint8 FSourceControlViewportMenu::GetOpacityValue(TWeakPtr<SLevelViewport> Viewport) const
 {
-	return SourceControlViewportUtils::GetFeedbackOpacity(ViewportClient);
+	if (TSharedPtr<SLevelViewport> Pinned = Viewport.Pin())
+	{
+		return SourceControlViewportUtils::GetFeedbackOpacity(&Pinned->GetLevelViewportClient());
+	}
+	return 0;
 }
 
 void FSourceControlViewportMenu::RecordToggleEvent(const FString& Param, bool bEnabled) const
diff --git a/Engine/Source/Developer/SourceControlViewport/Private/SourceControlViewportMenu.h b/Engine/Source/Developer/SourceControlViewport/Private/SourceControlViewportMenu.h
--- a/Engine/Source/Developer/SourceControlViewport/Private/SourceControlViewportMenu.h
+++ b/Engine/Source/Developer/SourceControlViewport/Private/SourceControlViewportMenu.h
@@ -4,7 +4,9 @@
 
 #include "CoreMinimal.h"
 #include "SourceControlViewportUtils.h"
+#include "Types/SlateEnums.h"
 
+class SLevelViewport;
 class FLevelEditorViewportClient;
 class UToolMenu;
 class SWidget;
@@ -23,21 +25,20 @@ public:
 private:
 	void InsertViewportMenu();
 	void PopulateViewportMenu(UToolMenu* InMenu);
+	void PopulateRevisionControlMenu(UToolMenu* InMenu);
 	void RemoveViewportMenu();
 
 private:
-	void ShowAll(FLevelEditorViewportClient* ViewportClient);
-	void HideAll(FLevelEditorViewportClient* ViewportClient);
+	void ShowAll(TWeakPtr<SLevelViewport> Viewport);
+	void HideAll(TWeakPtr<SLevelViewport> Viewport);
 
-	void ToggleHighlight(FLevelEditorViewportClient* ViewportClient, ESourceControlStatus Status);
-	bool IsHighlighted(FLevelEditorViewportClient* ViewportClient, ESourceControlStatus Status) const;
+	void ToggleHighlight(TWeakPtr<SLevelViewport> Viewport, ESourceControlStatus Status);
+	bool IsHighlighted(TWeakPtr<SLevelViewport> Viewport, ESourceControlStatus Status) const;
 
-	void SetOpacityValue(FLevelEditorViewportClient* ViewportClient, uint8 NewValue);
-	uint8 GetOpacityValue(FLevelEditorViewportClient* ViewportClient) const;
+	void OnOpacityCommitted(uint8 NewValue, ETextCommit::Type CommitType, TWeakPtr<SLevelViewport> Viewport);
+	void SetOpacityValue(uint8 NewValue, TWeakPtr<SLevelViewport> Viewport);
+	uint8 GetOpacityValue(TWeakPtr<SLevelViewport> Viewport) const;
 
 private:
 	void RecordToggleEvent(const FString& Param, bool bEnabled) const;
-
-private:
-	TSharedPtr<SWidget> OpacityWidget;
 };
\ No newline at end of file
diff --git a/Engine/Source/Editor/LevelEditor/Private/LevelEditorMenuContext.cpp b/Engine/Source/Editor/LevelEditor/Private/LevelEditorMenuContext.cpp
--- a/Engine/Source/Editor/LevelEditor/Private/LevelEditorMenuContext.cpp
+++ b/Engine/Source/Editor/LevelEditor/Private/LevelEditorMenuContext.cpp
@@ -4,6 +4,7 @@
 
 #include "Elements/Framework/TypedElementRegistry.h"
 #include "SLevelEditor.h"
+#include "ToolMenus.h"
 
 TWeakPtr<ILevelEditor> ULevelEditorMenuContext::GetLevelEditor() const
 {
@@ -14,3 +15,35 @@ FScriptTypedElementHandle ULevelEditorContextMenuContext::GetScriptHitProxyEleme
 {
 	return UTypedElementRegistry::GetInstance()->CreateScriptHandle(HitProxyElement.GetId());
 }
+
+TSharedPtr<SLevelViewport> ULevelViewportContext::GetLevelViewport(const UToolMenu* Menu)
+{
+	if (Menu)
+	{
+		return GetLevelViewport(Menu->Context);
+	}
+	return nullptr;
+}
+
+TSharedPtr<SLevelViewport> ULevelViewportContext::GetLevelViewport(const FToolMenuSection& Section)
+{
+	return GetLevelViewport(Section.Context);
+}
+
+TSharedPtr<SLevelViewport> ULevelViewportContext::GetLevelViewport(const FToolMenuContext& Context)
+{
+	if (const ULevelViewportContext* ViewportContext = Context.FindContext<ULevelViewportContext>())
+	{
+		return ViewportContext->LevelViewport.Pin();
+	}
+	return nullptr;
+}
+
+FLevelEditorViewportClient* ULevelViewportContext::GetLevelViewportClient() const
+{
+	if (TSharedPtr<SLevelViewport> Viewport = LevelViewport.Pin())
+	{
+		return &Viewport->GetLevelViewportClient();
+	}
+	return nullptr;
+}
diff --git a/Engine/Source/Editor/LevelEditor/Private/LevelEditorToolBar.cpp b/Engine/Source/Editor/LevelEditor/Private/LevelEditorToolBar.cpp
--- a/Engine/Source/Editor/LevelEditor/Private/LevelEditorToolBar.cpp
+++ b/Engine/Source/Editor/LevelEditor/Private/LevelEditorToolBar.cpp
@@ -1596,7 +1596,9 @@ void FLevelEditorToolBar::RegisterLevelEditorSecondaryModeToolbar()
 TSharedRef< SWidget > FLevelEditorToolBar::MakeLevelEditorToolBar( const TSharedRef<FUICommandList>& InCommandList, const TSharedRef<SLevelEditor> InLevelEditor )
 {
 	FLevelEditorModule& LevelEditorModule = FModuleManager::LoadModuleChecked<FLevelEditorModule>("LevelEditor");
+	PRAGMA_DISABLE_DEPRECATION_WARNINGS
 	FToolMenuContext MenuContext(InCommandList, LevelEditorModule.GetToolBarExtensibilityManager()->GetAllExtenders());
+	PRAGMA_ENABLE_DEPRECATION_WARNINGS
 	ULevelEditorMenuContext* LevelEditorMenuContext = NewObject<ULevelEditorMenuContext>();
 	LevelEditorMenuContext->LevelEditor = InLevelEditor;
 	MenuContext.AddObject(LevelEditorMenuContext);
diff --git a/Engine/Source/Editor/LevelEditor/Private/SLevelViewport.cpp b/Engine/Source/Editor/LevelEditor/Private/SLevelViewport.cpp
--- a/Engine/Source/Editor/LevelEditor/Private/SLevelViewport.cpp
+++ b/Engine/Source/Editor/LevelEditor/Private/SLevelViewport.cpp
@@ -41,7 +41,6 @@
 #include "UnrealEdMisc.h"
 #include "UnrealEdGlobals.h"
 #include "LevelEditor.h"
-#include "SLevelViewportToolBar.h"
 #include "LevelViewportActions.h"
 #include "LevelEditorActions.h"
 #include "SceneView.h"
@@ -229,7 +228,6 @@ SLevelViewport::SLevelViewport()
 	: HighResScreenshotDialog( nullptr )
 	, ViewTransitionType( EViewTransition::None )
 	, bViewTransitionAnimPending( false )
-	, DeviceProfile("Default")
 	, bPIEHasFocus(false)
 	, bPIEContainsFocus(false)
 	, UserAllowThrottlingValue(0)
@@ -386,17 +384,9 @@ void SLevelViewport::ConstructViewportOverlayContent()
 
 	// Add widget to show viewport warnings from the ActionableMessageSubsystem.
 	{
-		TAttribute<FMargin> Padding = TAttribute<FMargin>::CreateLambda(
-			[]()
-			{
-				const float TopPadding = UE::UnrealEd::ShowOldViewportToolbars() ? 37.0f : 8.0f;
-				return FMargin(0.f, TopPadding, 10.f, 0.f);
-			}
-		);
-
 		// clang-format off
 		ViewportOverlay->AddSlot(SlotIndex)
-		.Padding(Padding)
+		.Padding(FMargin(0.f, 8.f, 10.f, 0.f))
 		.VAlign(VAlign_Top)
 		.HAlign(HAlign_Right)
 		[
@@ -407,12 +397,7 @@ void SLevelViewport::ConstructViewportOverlayContent()
 				.Visibility_Lambda(
 					[this]()
 					{
-						if (!IsViewportToolbarVisible())
-						{
-							return EVisibility::Collapsed;
-						}
-						
-						if (GetToolBarVisibility() == EVisibility::Collapsed)
+						if (GetToolBarVisibility() != EVisibility::Visible)
 						{
 							return EVisibility::Collapsed;
 						}
@@ -680,9 +665,6 @@ void SLevelViewport::ConstructLevelEditorViewportClient(FLevelEditorViewportInst
 
 	bShowEditorToolbar = ViewportInstanceSettings.bShowFullToolbar;
 
-	// Always set to true initially
-	bShowToolbarAndControls = true; 
-
 	if (FPlatformMisc::IsRemoteSession())
 	{
 		// Bind to the change delegate of performance settings and call our handler with a dummy event to set current defaults
@@ -1313,7 +1295,7 @@ void SLevelViewport::Tick( const FGeometry& AllottedGeometry, const double InCur
 	if(IsPlayInEditorViewportActive() && bPIEHasFocus != ActiveViewport->HasMouseCapture())
 	{
 		bPIEHasFocus = ActiveViewport->HasMouseCapture();
-		if (!bPIEHasFocus || !UE::UnrealEd::ShowNewViewportToolbars())
+		if (!bPIEHasFocus)
 		{
 			PIEOverlayAnim = FCurveSequence(0.0f, SLevelViewportPIEAnimation::MouseControlLabelFadeout, ECurveEaseFunction::CubicInOut);
 			PIEOverlayAnim.Play(this->AsShared());
@@ -2188,8 +2170,13 @@ EVisibility SLevelViewport::OnGetViewportContentVisibility() const
 
 EVisibility SLevelViewport::GetToolBarVisibility() const
 {
-	// Do not show the toolbar if this viewport has a play in editor session, or we're in the VR Editor
-	return ( IsPlayInEditorViewportActive() || GEngine->IsStereoscopic3D( ActiveViewport.Get() ) ) ? EVisibility::Collapsed : OnGetViewportContentVisibility();
+	if (GEngine->IsStereoscopic3D(ActiveViewport.Get()))
+	{
+		// Never show the toolbar in VR
+		return EVisibility::Collapsed;
+	}
+
+	return IsViewportToolbarVisible() ? EVisibility::Visible : EVisibility::Collapsed;
 }
 
 EVisibility SLevelViewport::GetMaximizeToggleVisibility() const
@@ -2210,7 +2197,7 @@ EVisibility SLevelViewport::GetMaximizeToggleVisibility() const
 EVisibility SLevelViewport::GetCloseImmersiveButtonVisibility() const
 {
 	// Do not show the Immersive toggle button when not in immersive mode
-	return (IsImmersive() && bShowToolbarAndControls) ? EVisibility::Visible : EVisibility::Collapsed;
+	return IsImmersive() ? EVisibility::Visible : EVisibility::Collapsed;
 }
 
 EVisibility SLevelViewport::GetTransformToolbarVisibility() const
@@ -2249,41 +2236,6 @@ TSharedRef<FEditorViewportClient> SLevelViewport::MakeEditorViewportClient()
 	return LevelViewportClient.ToSharedRef();
 }
 
-TSharedPtr<SWidget> SLevelViewport::MakeViewportToolbar()
-{
-	const TSharedRef<SLevelViewportToolBar> OldViewportToolbar =
-		SAssignNew(LevelViewportToolBar, SLevelViewportToolBar)
-			.Viewport(SharedThis(this))
-			.Visibility_Lambda(
-				[this]() -> EVisibility
-				{
-					return GetToolBarVisibility();
-				}
-			)
-			.IsEnabled(FSlateApplication::Get().GetNormalExecutionAttribute());
-
-	return
-		// clang-format off
-		SNew(SVerticalBox)
-		.Visibility( EVisibility::SelfHitTestInvisible )
-		+SVerticalBox::Slot()
-		.AutoHeight()
-		.Padding(0, 1.0f, 0, 0)
-		.VAlign(VAlign_Top)
-		[
-			OldViewportToolbar
-		]
-		+SVerticalBox::Slot()
-		.VAlign(VAlign_Top)
-		.HAlign(HAlign_Left)
-		[
-			SNew(SActorPilotViewportToolbar)
-			.Viewport( SharedThis( this ) )
-			.Visibility(this, &SLevelViewport::GetLockedIconVisibility)
-		];
-	// clang-format on
-}
-
 TSharedPtr<SWidget> SLevelViewport::BuildViewportToolbar()
 {
 	const FName ViewportToolbarName = "LevelEditor.ViewportToolbar";
@@ -2405,6 +2357,7 @@ TSharedPtr<SWidget> SLevelViewport::BuildViewportToolbar()
 				ViewportToolbarMenu->AddSection("ScaleGridSnap").Alignment = EToolMenuSectionAlign::Last;
 				ViewportToolbarMenu->AddSection("CameraSpeed").Alignment = EToolMenuSectionAlign::Last;
 				
+				// TODO: Remove these once the deprecation period is complete
 				FToolMenuSection& FallbackSection = ViewportToolbarMenu->AddSection("Fallback");
 				FallbackSection.Alignment = EToolMenuSectionAlign::Last;
 				
@@ -2418,7 +2371,9 @@ TSharedPtr<SWidget> SLevelViewport::BuildViewportToolbar()
 					
 					UViewportToolBarContext* ExtensionContext = NewObject<UViewportToolBarContext>();
 					ExtensionContext->Viewport = Context->LevelViewport;
+					PRAGMA_DISABLE_DEPRECATION_WARNINGS
 					ExtensionContext->ViewportToolBar = Context->LevelViewportToolBarWidget;
+					PRAGMA_ENABLE_DEPRECATION_WARNINGS
 					
 					auto MakeExtensionPanel = [ExtensionContext, &Section](FName ExtensionId)
 					{
@@ -2494,8 +2449,10 @@ TSharedPtr<SWidget> SLevelViewport::BuildViewportToolbar()
 
 			// Stay backward-compatible with legacy toolbar menu extenders
 			{
+				PRAGMA_DISABLE_DEPRECATION_WARNINGS
 				TSharedPtr<FExtender> OldToolbarExtenders =
 					LevelEditorModule.GetToolBarExtensibilityManager()->GetAllExtenders();
+				PRAGMA_ENABLE_DEPRECATION_WARNINGS
 
 				if (OldToolbarExtenders.IsValid())
 				{
@@ -2519,14 +2476,14 @@ TSharedPtr<SWidget> SLevelViewport::BuildViewportToolbar()
 		}
 		
 		// Support legacy SExtensionPanel extensions
-		if (LevelViewportToolBar)
-		{
-			// Provide a non-null toolbar widget so that legacy extensions with bad null handling don't crash
-			ULegacyLevelViewportToolbarContext* const LegacyContext = NewObject<ULegacyLevelViewportToolbarContext>();
-			LegacyContext->LevelViewport = SharedThis(this);
-			LegacyContext->LevelViewportToolBarWidget = LevelViewportToolBar;
-			ViewportToolbarContext.AddObject(LegacyContext);
-		}
+		PRAGMA_DISABLE_DEPRECATION_WARNINGS
+		LevelViewportToolBar = SNew(SLevelViewportToolBar);
+		// Provide a non-null toolbar widget so that legacy extensions with bad null handling don't crash
+		ULegacyLevelViewportToolbarContext* const LegacyContext = NewObject<ULegacyLevelViewportToolbarContext>();
+		LegacyContext->LevelViewport = SharedThis(this);
+		LegacyContext->LevelViewportToolBarWidget = LevelViewportToolBar;
+		ViewportToolbarContext.AddObject(LegacyContext);
+		PRAGMA_ENABLE_DEPRECATION_WARNINGS
 
 		ExtendToolbarContext(ViewportToolbarContext);
 	}
@@ -2534,15 +2491,7 @@ TSharedPtr<SWidget> SLevelViewport::BuildViewportToolbar()
 	// clang-format off
 	return SNew(SWidgetSwitcher)
 	.WidgetIndex(this, &SLevelViewport::GetViewportToolbarIndex)
-	.Visibility_Lambda([this]
-	{
-		// TODO: Move this to GetToolBarVisibility() once the old toolbar is deprecated
-		if (GEngine->IsStereoscopic3D(ActiveViewport.Get()))
-		{
-			return EVisibility::Collapsed;
-		}
-		return IsViewportToolbarVisible() ? EVisibility::Visible : EVisibility::Collapsed;
-	})
+	.Visibility(this, &SLevelViewport::GetToolBarVisibility)
 	+ SWidgetSwitcher::Slot()
 	[
 		UToolMenus::Get()->GenerateWidget(ViewportToolbarName, ViewportToolbarContext)
@@ -2571,19 +2520,17 @@ TSharedRef<SWidget> SLevelViewport::BuildPIEViewportToolbar()
 				"DynamicPIEDetachHint",
 				FNewToolMenuSectionDelegate::CreateLambda([](FToolMenuSection& Section)
 				{
-					ULevelViewportContext* Context = Section.FindContext<ULevelViewportContext>();
-					if (!Context)
+					TSharedPtr<SLevelViewport> Viewport = ULevelViewportContext::GetLevelViewport(Section);
+					if (!Viewport)
 					{
 						return;
 					}
 					
-					TSharedPtr<SLevelViewport> Viewport = Context->LevelViewport.Pin();
-					
 					Section.AddEntry(FToolMenuEntry::InitWidget(
 						"PIEDetachHint",
 						SNew(SBox)
 						.Padding(FMargin(8.f, 0.f, 0.f, 0.f))
-						.Visibility_Lambda([WeakViewport = Context->LevelViewport]
+						.Visibility_Lambda([WeakViewport = Viewport.ToWeakPtr()]
 						{
 							const ULevelEditorPlaySettings* EditorPlayInSettings = GetDefault<ULevelEditorPlaySettings>();
 							if (!EditorPlayInSettings->ShowMouseControlLabel)
@@ -3532,12 +3479,7 @@ bool SLevelViewport::CanToggleMaximizeMode() const
 	return (ParentLayoutPinned.IsValid() && ParentLayoutPinned->IsMaximizeSupported() && !ParentLayoutPinned->IsTransitioning());
 }
 
-void  SLevelViewport::OnToggleMaximizeMode()
-{
-	OnToggleMaximize();
-}
-
-FReply SLevelViewport::OnToggleMaximize()
+void SLevelViewport::OnToggleMaximizeMode()
 {
 	TSharedPtr<FLevelViewportLayout> ParentLayoutPinned = ParentLayout.Pin();
 	if (ParentLayoutPinned.IsValid() && ParentLayoutPinned->IsMaximizeSupported())
@@ -3567,6 +3509,11 @@ FReply SLevelViewport::OnToggleMaximize()
 			ParentLayout.Pin()->RequestMaximizeViewport( ViewportName, bWantMaximize, bWantImmersive, bAllowAnimation );
 		}
 	}
+}
+
+FReply SLevelViewport::OnToggleMaximize()
+{
+	OnToggleMaximizeMode();
 	return FReply::Handled();
 }
 
@@ -3594,15 +3541,6 @@ void SLevelViewport::MakeImmersive( const bool bWantImmersive, const bool bAllow
 		{
 			ParentLayout.Pin()->RequestMaximizeViewport( ViewportName, bWantMaximize, bWantImmersive, bAllowAnimation );
 		}
-		
-		if (bWantImmersive && LevelViewportClient->IsVisualizeCalibrationMaterialEnabled()) 
-		{
-			bShowToolbarAndControls = false;
-		}
-		else
-		{
-			bShowToolbarAndControls = true;
-		}
 	}
 }
 
@@ -4678,21 +4616,6 @@ void SLevelViewport::OnPreviewSelectedCamerasChange()
 	}
 }
 
-void SLevelViewport::SetDeviceProfileString( const FString& ProfileName )
-{
-	DeviceProfile = ProfileName;
-}
-
-bool SLevelViewport::IsDeviceProfileStringSet( FString ProfileName ) const
-{
-	return DeviceProfile == ProfileName;
-}
-
-FString SLevelViewport::GetDeviceProfileString( ) const
-{
-	return DeviceProfile;
-}
-
 FText SLevelViewport::GetCurrentScreenPercentageText() const
 {
 	return FText::FromString(FString::Printf(TEXT("%3d%%"), int32(GetLevelViewportClient().GetPreviewScreenPercentage())));
@@ -4791,13 +4714,6 @@ EVisibility SLevelViewport::GetCurrentScreenPercentageVisibility() const
 	return Visible ? EVisibility::Visible : EVisibility::Collapsed;
 }
 
-EVisibility SLevelViewport::GetViewportControlsVisibility() const
-{
-	// Do not show the controls if this viewport has a play in editor session
-	// or is not the current viewport
-	return (&GetLevelViewportClient() == GCurrentLevelEditingViewportClient && !IsPlayInEditorViewportActive() && bShowToolbarAndControls) ? OnGetViewportContentVisibility() : EVisibility::Collapsed;
-}
-
 void SLevelViewport::OnSetViewportConfiguration(FName ConfigurationName)
 {
 	TSharedPtr<FLevelViewportLayout> LayoutPinned = ParentLayout.Pin();
@@ -4972,13 +4888,6 @@ void SLevelViewport::StartPlayInEditorSession(UGameViewportClient* InPlayClient,
 
 	bPIEHasFocus = ActiveViewport->HasMouseCapture();
 
-	if(EditorPlayInSettings->ShowMouseControlLabel && !GEngine->IsStereoscopic3D( ActiveViewport.Get() ) && !UE::UnrealEd::ShowNewViewportToolbars())
-	{
-		ELabelAnchorMode AnchorMode = EditorPlayInSettings->MouseControlLabelPosition.GetValue();
-
-		ShowMouseCaptureLabel(AnchorMode);
-	}
-
 	GEngine->BroadcastLevelActorListChanged();
 
 	// register for preview feature level change
@@ -5268,13 +5177,6 @@ void SLevelViewport::SwapViewportsForPlayInEditor()
 	// Put the mouse control label up again.
 	ULevelEditorPlaySettings const* EditorPlayInSettings = GetDefault<ULevelEditorPlaySettings>();
 	check(EditorPlayInSettings);
-	
-	if(EditorPlayInSettings->ShowMouseControlLabel && !GEngine->IsStereoscopic3D( ActiveViewport.Get() ) && !UE::UnrealEd::ShowNewViewportToolbars())
-	{
-		ELabelAnchorMode AnchorMode = EditorPlayInSettings->MouseControlLabelPosition.GetValue();
-		
-		ShowMouseCaptureLabel(AnchorMode);
-	}
 
 	// Swap between the active and inactive viewport
 	TSharedPtr<FSceneViewport> TempViewport = ActiveViewport;
@@ -5339,11 +5241,6 @@ void SLevelViewport::OnSimulateSessionFinished()
 	ActiveViewport->InvalidateHitProxy();
 }
 
-EVisibility SLevelViewport::GetLockedIconVisibility() const
-{
-	return IsAnyActorLocked() && UE::UnrealEd::ShowOldViewportToolbars() ? EVisibility::Visible : EVisibility::Collapsed;
-}
-
 FText SLevelViewport::GetLockedIconToolTip() const
 {
 	FText ToolTipText;
@@ -5568,16 +5465,6 @@ void SLevelViewport::OnFloatingButtonClicked()
 	LevelViewportClient->SetLastKeyViewport();
 }
 
-EVisibility  SLevelViewport::GetToolbarVisibility() const
-{
-	return GetOptionsMenuVisibility();
-}
-
-EVisibility SLevelViewport::GetOptionsMenuVisibility() const
-{
-	return bShowToolbarAndControls ? EVisibility::Visible : EVisibility::Collapsed;
-}
-
 void SLevelViewport::RemoveAllPreviews(const bool bRemoveFromDesktopViewport /**= true*/)
 {
 	// Clean up any actor preview viewports
diff --git a/Engine/Source/Editor/LevelEditor/Private/SLevelViewportToolBar.cpp b/Engine/Source/Editor/LevelEditor/Private/SLevelViewportToolBar.cpp
deleted file mode 100644
diff --git a/Engine/Source/Editor/LevelEditor/Private/SLevelViewportToolBar.h b/Engine/Source/Editor/LevelEditor/Private/SLevelViewportToolBar.h
deleted file mode 100644
diff --git a/Engine/Source/Editor/LevelEditor/Private/ViewportToolbar/LevelEditorSubmenus.cpp b/Engine/Source/Editor/LevelEditor/Private/ViewportToolbar/LevelEditorSubmenus.cpp
--- a/Engine/Source/Editor/LevelEditor/Private/ViewportToolbar/LevelEditorSubmenus.cpp
+++ b/Engine/Source/Editor/LevelEditor/Private/ViewportToolbar/LevelEditorSubmenus.cpp
@@ -315,8 +315,8 @@ FToolMenuEntry CreateEjectActorPilotEntry()
 		FNewToolMenuSectionDelegate::CreateLambda(
 			[](FToolMenuSection& InnerSection) -> void
 			{
-				ULevelViewportContext* const LevelViewportContext = InnerSection.FindContext<ULevelViewportContext>();
-				if (!LevelViewportContext)
+				TSharedPtr<::SLevelViewport> LevelViewport = ULevelViewportContext::GetLevelViewport(InnerSection);
+				if (!LevelViewport)
 				{
 					return;
 				}
@@ -324,7 +324,7 @@ FToolMenuEntry CreateEjectActorPilotEntry()
 				FToolUIAction EjectActorPilotAction;
 
 				EjectActorPilotAction.ExecuteAction = FToolMenuExecuteAction::CreateLambda(
-					[LevelViewportWeak = LevelViewportContext->LevelViewport](const FToolMenuContext& Context) -> void
+					[LevelViewportWeak = LevelViewport.ToWeakPtr()](const FToolMenuContext& Context) -> void
 					{
 						if (TSharedPtr<::SLevelViewport> LevelViewport = LevelViewportWeak.Pin())
 						{
@@ -334,7 +334,7 @@ FToolMenuEntry CreateEjectActorPilotEntry()
 				);
 
 				EjectActorPilotAction.CanExecuteAction = FToolMenuCanExecuteAction::CreateLambda(
-					[LevelViewportWeak = LevelViewportContext->LevelViewport](const FToolMenuContext& Context)
+					[LevelViewportWeak = LevelViewport.ToWeakPtr()](const FToolMenuContext& Context)
 					{
 						if (TSharedPtr<::SLevelViewport> EditorViewport = LevelViewportWeak.Pin())
 						{
@@ -363,14 +363,11 @@ FToolMenuEntry CreateEjectActorPilotEntry()
 				);
 
 				const TAttribute<bool> bShownInTopLevel = TAttribute<bool>::CreateLambda(
-					[ViewportContextWeak = TWeakObjectPtr<ULevelViewportContext>(LevelViewportContext)]() -> bool
+					[LevelViewportWeak = LevelViewport.ToWeakPtr()]() -> bool
 					{
-						if (TStrongObjectPtr<ULevelViewportContext> LevelViewportContext = ViewportContextWeak.Pin())
+						if (TSharedPtr<::SLevelViewport> LevelViewport = LevelViewportWeak.Pin())
 						{
-							if (const TSharedPtr<::SLevelViewport> LevelViewport = LevelViewportContext->LevelViewport.Pin())
-							{
-								return LevelViewport->GetLevelViewportClient().IsAnyActorLocked();
-							}
+							return LevelViewport->IsAnyActorLocked();
 						}
 
 						return true;
@@ -672,6 +669,8 @@ FToolMenuEntry CreatePreviewSelectedCamerasCheckBoxSubmenu()
 
 			constexpr float PreviewSizeMin = 1.0f;
 			constexpr float PreviewSizeMax = 10.0f;
+			
+			TSharedPtr<::SLevelViewport> LevelViewport = ULevelViewportContext::GetLevelViewport(Submenu);
 
 			FToolMenuEntry PreviewSizeEntry = UnrealEd::CreateNumericEntry(
 				"PreviewSize",
@@ -679,7 +678,7 @@ FToolMenuEntry CreatePreviewSelectedCamerasCheckBoxSubmenu()
 				LOCTEXT("PreviewSizeTooltip", "Affects the size of 'picture in picture' previews if they are enabled"),
 				FCanExecuteAction(),
 				UnrealEd::FNumericEntryExecuteActionDelegate::CreateLambda(
-					[Context = Submenu->Context](float InValue)
+					[LevelViewportWeak = LevelViewport.ToWeakPtr()](float InValue)
 					{
 						if (ULevelEditorViewportSettings* const ViewportSettings =
 								GetMutableDefault<ULevelEditorViewportSettings>())
@@ -689,12 +688,6 @@ FToolMenuEntry CreatePreviewSelectedCamerasCheckBoxSubmenu()
 							// If preview is not on, we assume the user wants to turn it on as they are editing preview size, so let's toggle it
 							if (!FLevelEditorActionCallbacks::IsPreviewSelectedCamerasChecked())
 							{
-								TWeakPtr<::SLevelViewport> LevelViewportWeak;
-								if (ULevelViewportContext* const LevelViewportContext =
-										Context.FindContext<ULevelViewportContext>())
-								{
-									LevelViewportWeak = LevelViewportContext->LevelViewport;
-								}
 								FLevelEditorActionCallbacks::TogglePreviewSelectedCameras(LevelViewportWeak);
 							}
 						}
@@ -731,13 +724,7 @@ FToolMenuEntry CreatePreviewSelectedCamerasCheckBoxSubmenu()
 		FToolUIAction(
 			FToolMenuExecuteAction::CreateLambda([](const FToolMenuContext& InContext)
 			{
-				TWeakPtr<::SLevelViewport> LevelViewportWeak;
-				if (ULevelViewportContext* const LevelViewportContext = InContext.FindContext<ULevelViewportContext>())
-				{
-					LevelViewportWeak = LevelViewportContext->LevelViewport;
-				}
-
-				FLevelEditorActionCallbacks::TogglePreviewSelectedCameras(LevelViewportWeak);
+				FLevelEditorActionCallbacks::TogglePreviewSelectedCameras(ULevelViewportContext::GetLevelViewport(InContext));
 			}),
 			FToolMenuGetActionCheckState::CreateLambda([](const FToolMenuContext& InContext)
 			{
@@ -902,19 +889,13 @@ TSharedPtr<FExtender> GetViewModesLegacyExtenders()
 
 void PopulateViewModesMenu(UToolMenu* InMenu)
 {
-	FToolMenuInsert InsertPosition("ViewMode", EToolMenuInsertType::After);
-
-	ULevelViewportContext* const LevelViewportContext = InMenu->FindContext<ULevelViewportContext>();
-	if (!LevelViewportContext)
-	{
-		return;
-	}
-
-	const TSharedPtr<::SLevelViewport> LevelViewport = LevelViewportContext->LevelViewport.Pin();
+	const TSharedPtr<::SLevelViewport> LevelViewport = ULevelViewportContext::GetLevelViewport(InMenu);
 	if (!LevelViewport)
 	{
 		return;
 	}
+	
+	const FToolMenuInsert InsertPosition("ViewMode", EToolMenuInsertType::After);
 
 	{
 		FToolMenuSection& Section = InMenu->AddSection(
@@ -1329,13 +1310,7 @@ FToolMenuEntry CreateShowFoliageSubmenu()
 		FNewToolMenuDelegate::CreateLambda(
 			[](UToolMenu* Submenu)
 			{
-				ULevelViewportContext* const LevelViewportContext = Submenu->FindContext<ULevelViewportContext>();
-				if (!LevelViewportContext)
-				{
-					return;
-				}
-
-				TSharedPtr<::SLevelViewport> Viewport = LevelViewportContext->LevelViewport.Pin();
+				TSharedPtr<::SLevelViewport> Viewport = ULevelViewportContext::GetLevelViewport(Submenu);
 				if (!Viewport)
 				{
 					return;
@@ -1416,13 +1391,7 @@ FToolMenuEntry CreateShowHLODsSubmenu()
 		FNewToolMenuSectionDelegate::CreateLambda(
 			[](FToolMenuSection& InDynamicSection) -> void
 			{
-				ULevelViewportContext* const LevelViewportContext = InDynamicSection.FindContext<ULevelViewportContext>();
-				if (!LevelViewportContext)
-				{
-					return;
-				}
-
-				TSharedPtr<::SLevelViewport> Viewport = LevelViewportContext->LevelViewport.Pin();
+				TSharedPtr<::SLevelViewport> Viewport = ULevelViewportContext::GetLevelViewport(InDynamicSection);
 				if (!Viewport)
 				{
 					return;
@@ -1447,14 +1416,7 @@ FToolMenuEntry CreateShowHLODsSubmenu()
 					FNewToolMenuDelegate::CreateLambda(
 						[](UToolMenu* Submenu)
 						{
-							ULevelViewportContext* const LevelViewportContext =
-								Submenu->FindContext<ULevelViewportContext>();
-							if (!LevelViewportContext)
-							{
-								return;
-							}
-
-							TSharedPtr<::SLevelViewport> Viewport = LevelViewportContext->LevelViewport.Pin();
+							TSharedPtr<::SLevelViewport> Viewport = ULevelViewportContext::GetLevelViewport(Submenu);
 							if (!Viewport)
 							{
 								return;
@@ -1716,13 +1678,7 @@ FToolMenuEntry CreateShowLayersSubmenu()
 		FNewToolMenuSectionDelegate::CreateLambda(
 			[](FToolMenuSection& InDynamicSection) -> void
 			{
-				ULevelViewportContext* const LevelViewportContext = InDynamicSection.FindContext<ULevelViewportContext>();
-				if (!LevelViewportContext)
-				{
-					return;
-				}
-
-				TSharedPtr<::SLevelViewport> Viewport = LevelViewportContext->LevelViewport.Pin();
+				TSharedPtr<::SLevelViewport> Viewport = ULevelViewportContext::GetLevelViewport(InDynamicSection);
 				if (!Viewport)
 				{
 					return;
@@ -1815,17 +1771,10 @@ FToolMenuEntry CreateShowStatsSubmenu(bool bInAddToggleStatsCheckbox, TAttribute
 		FNewToolMenuSectionDelegate::CreateLambda(
 			[bInAddToggleStatsCheckbox, InLabelOverride](FToolMenuSection& InDynamicSection)
 			{
-				TWeakPtr<::SLevelViewport> WeakLevelViewport;
-				if (ULevelViewportContext* const LevelViewportContext =
-						InDynamicSection.FindContext<ULevelViewportContext>())
-				{
-					WeakLevelViewport = LevelViewportContext->LevelViewport;
-				}
-
 				FToolUIActionChoice CommandAction;
 				if (bInAddToggleStatsCheckbox)
 				{
-					if (TSharedPtr<::SLevelViewport> Viewport = WeakLevelViewport.Pin())
+					if (TSharedPtr<::SLevelViewport> Viewport = ULevelViewportContext::GetLevelViewport(InDynamicSection))
 					{
 						CommandAction =
 							FToolUIActionChoice(FEditorViewportCommands::Get().ToggleStats, *Viewport->GetCommandList());
@@ -2274,13 +2223,7 @@ FToolMenuEntry CreatePIEPerformanceAndScalabilitySubmenu()
 
 void GenerateViewportLayoutsMenu(UToolMenu* InMenu)
 {
-	ULevelViewportContext* const LevelViewportContext = InMenu->FindContext<ULevelViewportContext>();
-	if (!LevelViewportContext)
-	{
-		return;
-	}
-
-	const TSharedPtr<::SLevelViewport> LevelViewport = LevelViewportContext->LevelViewport.Pin();
+	const TSharedPtr<::SLevelViewport> LevelViewport = ULevelViewportContext::GetLevelViewport(InMenu);
 	if (!LevelViewport)
 	{
 		return;
@@ -2668,13 +2611,7 @@ FToolMenuEntry CreateViewportSizingSubmenu()
 		FNewToolMenuDelegate::CreateLambda(
 			[](UToolMenu* Submenu) -> void
 			{
-				ULevelViewportContext* const LevelViewportContext = Submenu->FindContext<ULevelViewportContext>();
-				if (!LevelViewportContext)
-				{
-					return;
-				}
-
-				if (const TSharedPtr<::SLevelViewport> LevelViewport = LevelViewportContext->LevelViewport.Pin())
+				if (const TSharedPtr<::SLevelViewport> LevelViewport = ULevelViewportContext::GetLevelViewport(Submenu))
 				{
 					GenerateViewportLayoutsMenu(Submenu);
 				}
@@ -2780,20 +2717,14 @@ void CreateCameraSpawnMenu(UToolMenu* InMenu)
 
 void CreateBookmarksMenu(UToolMenu* InMenu)
 {
-	ULevelViewportContext* const LevelViewportContext = InMenu->FindContext<ULevelViewportContext>();
-	if (!LevelViewportContext)
-	{
-		return;
-	}
-
-	const TSharedPtr<::SLevelViewport> LevelViewportPinned = LevelViewportContext->LevelViewport.Pin();
-	if (!LevelViewportPinned)
+	const TSharedPtr<::SLevelViewport> LevelViewport = ULevelViewportContext::GetLevelViewport(InMenu);
+	if (!LevelViewport)
 	{
 		return;
 	}
 
 	// Add a menu entry for each bookmark
-	const FLevelEditorViewportClient& ViewportClient = LevelViewportPinned->GetLevelViewportClient();
+	const FLevelEditorViewportClient& ViewportClient = LevelViewport->GetLevelViewportClient();
 
 	FToolMenuSection& ManageBookmarksSection =
 		InMenu->FindOrAddSection("ManageBookmarks", LOCTEXT("ManageBookmarkSectionName", "Manage Bookmarks"));
@@ -2846,7 +2777,7 @@ void CreateBookmarksMenu(UToolMenu* InMenu)
 					LOCTEXT("ManageBookmarksSubMenu", "Manage Bookmarks"),
 					LOCTEXT("ManageBookmarksSubMenu_ToolTip", "Bookmarks related actions"),
 					FNewToolMenuDelegate::CreateLambda(
-						[bFoundBookmarks, LevelViewportWeak = LevelViewportContext->LevelViewport](UToolMenu* InMenu)
+						[bFoundBookmarks, LevelViewportWeak = LevelViewport.ToWeakPtr()](UToolMenu* InMenu)
 						{
 							if (!bFoundBookmarks)
 							{
@@ -2960,13 +2891,7 @@ FToolMenuEntry CreateFarViewPlaneMenu(TWeakPtr<::SLevelViewport> InLevelViewport
 
 void AddCameraActorSelectSection(UToolMenu* InMenu)
 {
-	ULevelViewportContext* const LevelViewportContext = InMenu->FindContext<ULevelViewportContext>();
-	if (!LevelViewportContext)
-	{
-		return;
-	}
-
-	TSharedPtr<::SLevelViewport> LevelViewport = LevelViewportContext->LevelViewport.Pin();
+	TSharedPtr<::SLevelViewport> LevelViewport = ULevelViewportContext::GetLevelViewport(InMenu);
 	if (!LevelViewport)
 	{
 		return;
@@ -3003,7 +2928,7 @@ void AddCameraActorSelectSection(UToolMenu* InMenu)
 			CameraActorsHeading,
 			LOCTEXT("LookThroughPlacedCameras_ToolTip", "Look through and pilot placed cameras"),
 			FNewToolMenuDelegate::CreateLambda(
-				[LookThroughActors, LevelViewportWeak = LevelViewportContext->LevelViewport](UToolMenu* InMenu)
+				[LookThroughActors, LevelViewportWeak = LevelViewport.ToWeakPtr()](UToolMenu* InMenu)
 				{
 					if (TSharedPtr<::SLevelViewport> LevelViewport = LevelViewportWeak.Pin())
 					{
@@ -3041,8 +2966,7 @@ void ExtendCameraSubmenu(FName InCameraOptionsSubmenuName)
 		FNewToolMenuDelegate::CreateLambda(
 			[](UToolMenu* InDynamicMenu)
 			{
-				ULevelViewportContext* LevelViewportContext = InDynamicMenu->FindContext<ULevelViewportContext>();
-				TWeakPtr<::SLevelViewport> LevelViewportWeak = LevelViewportContext->LevelViewport;
+				TWeakPtr<::SLevelViewport> LevelViewportWeak = ULevelViewportContext::GetLevelViewport(InDynamicMenu);
 
 				// Camera Selection elements
 				{
diff --git a/Engine/Source/Editor/LevelEditor/Private/ViewportToolbar/LevelViewportContext.cpp b/Engine/Source/Editor/LevelEditor/Private/ViewportToolbar/LevelViewportContext.cpp
--- a/Engine/Source/Editor/LevelEditor/Private/ViewportToolbar/LevelViewportContext.cpp
+++ b/Engine/Source/Editor/LevelEditor/Private/ViewportToolbar/LevelViewportContext.cpp
@@ -11,5 +11,7 @@ FLevelEditorViewportClient* ULegacyLevelViewportToolbarContext::GetLevelViewport
 		return &Viewport->GetLevelViewportClient();
 	}
 	
+	PRAGMA_DISABLE_DEPRECATION_WARNINGS
 	return Super::GetLevelViewportClient();
+	PRAGMA_ENABLE_DEPRECATION_WARNINGS
 }
diff --git a/Engine/Source/Editor/LevelEditor/Private/ViewportToolbar/LevelViewportContext.h b/Engine/Source/Editor/LevelEditor/Private/ViewportToolbar/LevelViewportContext.h
--- a/Engine/Source/Editor/LevelEditor/Private/ViewportToolbar/LevelViewportContext.h
+++ b/Engine/Source/Editor/LevelEditor/Private/ViewportToolbar/LevelViewportContext.h
@@ -8,17 +8,10 @@
 
 #include "LevelViewportContext.generated.h"
 
-#define UE_API LEVELEDITOR_API
-
-UCLASS(MinimalAPI)
-class ULevelViewportContext : public UObject
-{
-	GENERATED_BODY()
 
-public:
-	TWeakPtr<SLevelViewport> LevelViewport;
-};
+#define UE_API LEVELEDITOR_API
 
+PRAGMA_DISABLE_DEPRECATION_WARNINGS
 // Customizes ULevelViewportToolBarContext for backwards compatability
 // TODO: Remove this once the old toolbars are deprecated.
 UCLASS(MinimalAPI)
@@ -31,5 +24,6 @@ public:
 	
 	UE_API virtual FLevelEditorViewportClient* GetLevelViewportClient() const override;
 };
+PRAGMA_ENABLE_DEPRECATION_WARNINGS
 
 #undef UE_API
diff --git a/Engine/Source/Editor/LevelEditor/Public/LevelEditorMenuContext.h b/Engine/Source/Editor/LevelEditor/Public/LevelEditorMenuContext.h
--- a/Engine/Source/Editor/LevelEditor/Public/LevelEditorMenuContext.h
+++ b/Engine/Source/Editor/LevelEditor/Public/LevelEditorMenuContext.h
@@ -28,6 +28,9 @@ class SLevelViewportToolBar;
 class UActorComponent;
 class UTypedElementSelectionSet;
 struct FFrame;
+class UToolMenu;
+struct FToolMenuSection;
+struct FToolMenuContext;
 
 UCLASS(MinimalAPI)
 class ULevelEditorMenuContext : public UObject
@@ -83,13 +86,31 @@ public:
 };
 
 UCLASS(MinimalAPI)
-class ULevelViewportToolBarContext : public UObject
+class ULevelViewportContext : public UObject
+{
+	GENERATED_BODY()
+
+public:
+	static UE_API TSharedPtr<SLevelViewport> GetLevelViewport(const UToolMenu* Menu);
+	static UE_API TSharedPtr<SLevelViewport> GetLevelViewport(const FToolMenuSection& Section);
+	static UE_API TSharedPtr<SLevelViewport> GetLevelViewport(const FToolMenuContext& Context);
+
+	TWeakPtr<SLevelViewport> LevelViewport;
+	
+	UE_API FLevelEditorViewportClient* GetLevelViewportClient() const;
+};
+
+
+UCLASS(MinimalAPI)
+class UE_DEPRECATED(5.7, "Use ULevelViewportContext instead") ULevelViewportToolBarContext : public UObject
 {
 	GENERATED_BODY()
 public:
+	PRAGMA_DISABLE_DEPRECATION_WARNINGS
 	TWeakPtr<SLevelViewportToolBar> LevelViewportToolBarWidget;
+	PRAGMA_ENABLE_DEPRECATION_WARNINGS
 
-	UE_API virtual FLevelEditorViewportClient* GetLevelViewportClient() const;
+	UE_API virtual FLevelEditorViewportClient* GetLevelViewportClient() const { return nullptr; }
 };
 
 
diff --git a/Engine/Source/Editor/LevelEditor/Public/SLevelViewport.h b/Engine/Source/Editor/LevelEditor/Public/SLevelViewport.h
--- a/Engine/Source/Editor/LevelEditor/Public/SLevelViewport.h
+++ b/Engine/Source/Editor/LevelEditor/Public/SLevelViewport.h
@@ -20,6 +20,7 @@
 #include "EditorModeManager.h"
 #include "IAssetViewport.h"
 #include "LevelEditorViewport.h"
+#include "SViewportToolBar.h"
 
 #define UE_API LEVELEDITOR_API
 
@@ -37,6 +38,15 @@ class UFoliageType;
 enum class EMapChangeType : uint8;
 enum ELabelAnchorMode : int;
 
+// This implementation is only present so that ULevelViewportToolBarContext & UViewportToolbarContext continues to have a valid public pointer.
+class UE_DEPRECATED(5.7, "Use the UToolMenu \"LevelEditor.ViewportToolbar\" menu instead.") SLevelViewportToolBar : public SViewportToolBar
+{
+	SLATE_BEGIN_ARGS(SLevelViewportToolBar) {}
+	SLATE_END_ARGS()
+	
+	void Construct(const FArguments& InArgs) {};
+};
+
 /**
  * Encapsulates an SViewport and an SLevelViewportToolBar
  */
@@ -142,6 +152,7 @@ public:
 	/**
 	 * Called when the maximize command is executed                   
 	 */
+	UE_DEPRECATED(5.7, "Use OnTogleMaximizeMode() instead.")	
 	UE_API FReply OnToggleMaximize();
 
 	/**
@@ -174,11 +185,13 @@ public:
 	 */
 	UE_API bool IsImmersive () const override;
 	
+	UE_DEPRECATED(5.7, "Immersive mode is now handled differently")
 	/**
 	 * Called to get the visibility of the viewport's 'Restore from Immersive' button. Returns EVisibility::Collapsed when not in immersive mode
 	 */
 	UE_API EVisibility GetCloseImmersiveButtonVisibility() const;
 		
+	UE_DEPRECATED(5.7, "Maximuize & Immersive mode is now handled through the new viewport toolbar.")
 	/**
 	 * Called to get the visibility of the viewport's maximize/minimize toggle button. Returns EVisibility::Collapsed when in immersive mode
 	 */
@@ -321,17 +334,20 @@ public:
 	 *
 	 * @param ProfileName The profile name to set
 	 */
-	UE_API void SetDeviceProfileString( const FString& ProfileName );
+	UE_DEPRECATED(5.7, "Device profiles & preview platforms are handled in the new ToolMenus toolbar.")
+	void UE_API SetDeviceProfileString( const FString& ProfileName ) {}
 
 	/**
 	 * @return true if the in profile name matches the set profile name
 	 */
-	UE_API bool IsDeviceProfileStringSet( FString ProfileName ) const;
+	UE_DEPRECATED(5.7, "Device profiles & preview platforms are handled in the new ToolMenus toolbar.")
+	bool UE_API IsDeviceProfileStringSet( FString ProfileName ) const { return false; }
 
 	/**
 	 * @return the name of the selected device profile
 	 */
-	UE_API FString GetDeviceProfileString( ) const;
+	UE_DEPRECATED(5.7, "Device profiles & preview platforms are handled in the new ToolMenus toolbar.")
+	FString UE_API GetDeviceProfileString( ) const { return FString(); }
 
 	/** Get the parent level editor for this viewport */
 	TWeakPtr<ILevelEditor> GetParentLevelEditor() const { return ParentLevelEditor; }
@@ -356,8 +372,9 @@ public:
 	/** @return The visibility of the current screen percentage text display */
 	UE_API EVisibility GetCurrentScreenPercentageVisibility() const;
 
+	UE_DEPRECATED(5.7, "This function is no longer used.")
 	/** @return The visibility of the viewport controls popup */
-	UE_API virtual EVisibility GetViewportControlsVisibility() const;
+	UE_API virtual EVisibility GetViewportControlsVisibility() const { return EVisibility::Collapsed; }
 
 	/**
 	 * Called to get the visibility of the level viewport toolbar
@@ -419,13 +436,15 @@ public:
 	UE_API void OnFloatingButtonClicked();
 
 	UE_DEPRECATED(5.6, "Call GetOptionsMenuVisibility() instead.")
-	UE_API EVisibility GetToolbarVisibility() const;
+	UE_API EVisibility GetToolbarVisibility() const { return EVisibility::Visible;}
 
+	UE_DEPRECATED(5.7, "Toolbar visibility is no longer controlled in partial steps.")
 	/** Get the visibility for viewport toolbar's options menu. */
-	UE_API EVisibility GetOptionsMenuVisibility() const;
+	UE_API EVisibility GetOptionsMenuVisibility() const { return EVisibility::Visible; }
 
+	UE_DEPRECATED(5.7, "Toolbar visibility is no longer controlled in partial steps.")
 	/** Get the visibility for items considered to be part of the 'full' viewport toolbar */
-	EVisibility GetFullToolbarVisibility() const { return (bShowToolbarAndControls && bShowEditorToolbar) ? EVisibility::SelfHitTestInvisible : EVisibility::Collapsed; }
+	EVisibility GetFullToolbarVisibility() const { return EVisibility::SelfHitTestInvisible; }
 
 	/** Unpin and close all actor preview windows */
 	UE_API void RemoveAllPreviews(const bool bRemoveFromDesktopViewport = true);
@@ -488,7 +507,6 @@ public:
 protected:
 	/** SEditorViewport interface */
 	UE_API virtual TSharedRef<FEditorViewportClient> MakeEditorViewportClient() override;
-	UE_API virtual TSharedPtr<SWidget> MakeViewportToolbar() override;
 	UE_API virtual TSharedPtr<SWidget> BuildViewportToolbar() override;
 	UE_API TSharedRef<SWidget> BuildPIEViewportToolbar();
 
@@ -799,9 +817,6 @@ private:
 	
 	UE_API void OnEditorClose();
 
-	/** Gets whether the locked icon should be shown in the viewport because it is locked to an actor */
-	UE_API EVisibility GetLockedIconVisibility() const;
-
 	/** Gets the locked icon tooltip text showing the meaning of the icon and the name of the locked actor */
 	UE_API FText GetLockedIconToolTip() const;
 
@@ -911,7 +926,9 @@ private:
 	 */
 	FDelegateHandle PIEPreviewFeatureLevelChangedHandle;
 
+	PRAGMA_DISABLE_DEPRECATION_WARNINGS
 	TSharedPtr<SLevelViewportToolBar> LevelViewportToolBar;
+	PRAGMA_ENABLE_DEPRECATION_WARNINGS
 
 	/** Level viewport client */
 	TSharedPtr<FLevelEditorViewportClient> LevelViewportClient;
@@ -977,9 +994,6 @@ private:
 	/** True if we want to kick off a transition animation but are waiting for the next tick to do so */
 	bool bViewTransitionAnimPending;
 
-	/** The current device profile string */
-	FString DeviceProfile;
-
 	/** The current viewport config key */
 	FName ConfigKey;
 
@@ -1046,9 +1060,6 @@ private:
 	/** The users value for allowing throttling, we restore this value when we lose focus. */
 	int32 UserAllowThrottlingValue;
 
-	/** Whether to show toolbar or buttons */
-	bool bShowToolbarAndControls;
-
 	/** Whether to show the editor toolbar */
 	bool bShowEditorToolbar = true;
 	
