SourceControl / ViewportToolTips: Delay calling GetHitProxy until past the tooltip delay.

A Superluminal capture shows delays can happen when calling GetHitProxy before the cache is updated.
This could be potential the cause for lag / hiccups in the viewport reported by users when URC is enabled.

#rnx

diff --git a/Engine/Source/Developer/SourceControlViewport/Private/SourceControlViewportToolTips.cpp b/Engine/Source/Developer/SourceControlViewport/Private/SourceControlViewportToolTips.cpp
--- a/Engine/Source/Developer/SourceControlViewport/Private/SourceControlViewportToolTips.cpp
+++ b/Engine/Source/Developer/SourceControlViewport/Private/SourceControlViewportToolTips.cpp
@@ -84,31 +84,33 @@ void FSourceControlViewportToolTips::UpdateCanvas(float DeltaTime)
 				{
 					Actor.Reset();
 					UpdateToolTip();
+					DelayTime = 0;
 
 					ActorMouseX = MouseX;
 					ActorMouseY = MouseY;
 				}
-
-				HActor* ActorHitProxy = HitProxyCast<HActor>(Viewport->GetHitProxy(MouseX, MouseY));
-				if (ActorHitProxy == nullptr)
+				else
 				{
-					Actor.Reset();
-					UpdateToolTip();
+					DelayTime += DeltaTime;
 				}
-				if (ActorHitProxy != nullptr && ActorHitProxy->Actor != Actor)
+
+				if (DelayTime >= 0.15f) // See: Slate.TooltipSummonDelay
 				{
-					Actor.Reset();
-					UpdateToolTip();
+					HActor* ActorHitProxy = HitProxyCast<HActor>(Viewport->GetHitProxy(MouseX, MouseY));
+					if (ActorHitProxy == nullptr)
+					{
+						Actor.Reset();
+						UpdateToolTip();
+					}
+					if (ActorHitProxy != nullptr && ActorHitProxy->Actor != Actor)
+					{
+						Actor.Reset();
+						UpdateToolTip();
 
-					Actor = MakeWeakObjectPtr(ActorHitProxy->Actor);
-					ActorTime = FPlatformTime::Seconds();
-					DelayTime = 0;
-				}
+						Actor = MakeWeakObjectPtr(ActorHitProxy->Actor);
+					}
 
-				if (Actor.IsValid())
-				{
-					DelayTime += DeltaTime;
-					if (DelayTime >= 0.15f) // See: Slate.TooltipSummonDelay
+					if (Actor.IsValid())
 					{
 						UpdateToolTip();
 					}
diff --git a/Engine/Source/Developer/SourceControlViewport/Private/SourceControlViewportToolTips.h b/Engine/Source/Developer/SourceControlViewport/Private/SourceControlViewportToolTips.h
--- a/Engine/Source/Developer/SourceControlViewport/Private/SourceControlViewportToolTips.h
+++ b/Engine/Source/Developer/SourceControlViewport/Private/SourceControlViewportToolTips.h
@@ -57,8 +57,6 @@ private:
 	int32 ActorMouseX = 0;
 	int32 ActorMouseY = 0;
 
-	/* The time at which the actor became relevant */
-	double ActorTime = 0;
 	/* The time since the actor became relevant */
 	float DelayTime = 0;
 
